Removed dialog box about switching sessisons because it's bad design.
Discussed with [REDACTED_PERSON]

diff --git a/Engine/Source/Developer/Profiler/Private/ProfilerManager.cpp b/Engine/Source/Developer/Profiler/Private/ProfilerManager.cpp
--- a/Engine/Source/Developer/Profiler/Private/ProfilerManager.cpp
+++ b/Engine/Source/Developer/Profiler/Private/ProfilerManager.cpp
@@ -380,11 +380,6 @@ void FProfilerManager::ProfilerClient_OnClientDisconnected( const FGuid& Session
 
 void FProfilerManager::SessionManager_OnCanSelectSession( const ISessionInfoPtr& SelectedSession, bool& CanSelect )
 {
-	if( SelectedSession.IsValid() && ActiveSession != SelectedSession && ActiveSession.IsValid() )
-	{
-		EAppReturnType::Type Result = FPlatformMisc::MessageBoxExt( EAppMsgType::YesNo, *LOCTEXT("CanSelectSessionQuestion", "Do you really want to switch to this session?").ToString(), *LOCTEXT("CanSelectSessionCaption", "Warning").ToString() );
-		CanSelect = Result == EAppReturnType::Yes ? true : false;
-	}
 }
 
 void FProfilerManager::SessionManager_OnSelectedSessionChanged( const ISessionInfoPtr& InActiveSession )
