#UE4.Profiler
- Added first version of the profiler mini-view
- Switched to SCOPE_LOG_TIME when measuring profiler's performance
- Misc fixed how the stats metadata is handled by the profiler
- Game thread is tracked by default

diff --git a/Engine/Source/Developer/Profiler/Private/ProfilerDataSource.cpp b/Engine/Source/Developer/Profiler/Private/ProfilerDataSource.cpp
--- a/Engine/Source/Developer/Profiler/Private/ProfilerDataSource.cpp
+++ b/Engine/Source/Developer/Profiler/Private/ProfilerDataSource.cpp
@@ -394,7 +394,7 @@ FEventGraphData::FEventGraphData( const FProfilerSessionRef& InProfilerSession,
 	, ProfilerSessionPtr( InProfilerSession )
 {
 	static FTotalTimeAndCount Current(0.0f, 0);
-	FProfilerScopedLogTime ScopedLog( TEXT("FEventGraphData::FEventGraphData"), &Current );
+	PROFILER_SCOPE_LOG_TIME( TEXT( "FEventGraphData::FEventGraphData" ), &Current );
 
 	Description = FString::Printf( TEXT("%s: %i"), *InProfilerSession->GetShortName(), InFrameIndex );
 
diff --git a/Engine/Source/Developer/Profiler/Private/ProfilerManager.cpp b/Engine/Source/Developer/Profiler/Private/ProfilerManager.cpp
--- a/Engine/Source/Developer/Profiler/Private/ProfilerManager.cpp
+++ b/Engine/Source/Developer/Profiler/Private/ProfilerManager.cpp
@@ -194,11 +194,12 @@ void FProfilerManager::LoadProfilerCapture( const FString& ProfilerCaptureFilepa
 	const FGuid ProfilerInstanceID = ProfilerSession->GetInstanceID();
 
 	ProfilerSession->
-		SetOnCaptureFileProcessed( FProfilerSession::FCaptureFileProcessedDelegate::CreateSP( this, &FProfilerManager::ProfilerSession_OnCaptureFileProcessed ) );
+		SetOnCaptureFileProcessed( FProfilerSession::FCaptureFileProcessedDelegate::CreateSP( this, &FProfilerManager::ProfilerSession_OnCaptureFileProcessed ) )
+		.SetOnAddThreadTime( FProfilerSession::FAddThreadTimeDelegate::CreateSP( this, &FProfilerManager::ProfilerSession_OnAddThreadTime ) );
 
 	ProfilerSessionInstances.Add( ProfilerInstanceID, ProfilerSession );
 	{
-		FProfilerScopedLogTime LC( TEXT("ProfilerClient->LoadCapture") );
+		PROFILER_SCOPE_LOG_TIME( TEXT( "ProfilerClient->LoadCapture" ), );
 		ProfilerClient->LoadCapture( ProfilerCaptureFilepath, ProfilerInstanceID );
 	}
 
@@ -211,7 +212,6 @@ void FProfilerManager::LoadProfilerCapture( const FString& ProfilerCaptureFilepa
 
 void FProfilerManager::LoadRawStatsFile( const FString& RawStatsFileFileath )
 {
-	// deselect the active session
 	if( ActiveSession.IsValid() )
 	{
 		SessionManager->SelectSession( NULL );
@@ -227,8 +227,9 @@ void FProfilerManager::LoadRawStatsFile( const FString& RawStatsFileFileath )
 	ProfilerSession_OnCaptureFileProcessed( ProfilerInstanceID );
 	TrackDefaultStats();
 
-// 	ProfilerSession->
-// 		SetOnCaptureFileProcessed( FProfilerSession::FCaptureFileProcessedDelegate::CreateSP( this, &FProfilerManager::ProfilerSession_OnCaptureFileProcessed ) );
+	ProfilerSession->
+		SetOnCaptureFileProcessed( FProfilerSession::FCaptureFileProcessedDelegate::CreateSP( this, &FProfilerManager::ProfilerSession_OnCaptureFileProcessed ) )
+		.SetOnAddThreadTime( FProfilerSession::FAddThreadTimeDelegate::CreateSP( this, &FProfilerManager::ProfilerSession_OnAddThreadTime ) );
 
 	ProfilerSessionInstances.Add( ProfilerInstanceID, ProfilerSession );
 
@@ -288,6 +289,8 @@ void FProfilerManager::ProfilerClient_OnProfilerData( const FGuid& InstanceID, c
 	if( ProfilerSession && ProfilerWindow.IsValid())
 	{
 		(*ProfilerSession)->UpdateProfilerData( Content );
+		// Game thread should always be enabled.
+		TrackDefaultStats();
 
 		// Update the notification that a file is being loaded.
 		GetProfilerWindow()->ManageLoadingProgressNotificationState( (*ProfilerSession)->GetName(), EProfilerNotificationTypes::LoadingOfflineCapture, ELoadingProgressStates::InProgress, DataLoadingProgress );
@@ -425,6 +428,7 @@ void FProfilerManager::SessionManager_OnInstanceSelectionChanged()
 			if( !bAlreadyAdded )
 			{
 				FProfilerSessionRef ProfilerSession = MakeShareable( new FProfilerSession( SessionInstanceInfo ) );
+				ProfilerSession->SetOnAddThreadTime( FProfilerSession::FAddThreadTimeDelegate::CreateSP( this, &FProfilerManager::ProfilerSession_OnAddThreadTime ) );
 
 				ProfilerSessionInstances.Add( ProfilerSession->GetInstanceID(), ProfilerSession );
 				ProfilerClient->Track( ProfilerInstanceID );
@@ -543,6 +547,7 @@ bool FProfilerManager::UntrackStat( const uint32 StatID )
 void FProfilerManager::ClearStatsAndInstances()
 {
 	CloseAllEventGraphTabs();
+	GetProfilerWindow()->ProfilerMiniView->Reset();
 
 	ProfilerType = EProfilerSessionTypes::InvalidOrMax;
 	SetDataPreview( false );
@@ -683,8 +688,13 @@ const FLinearColor& FProfilerManager::GetColorForStatID( const uint32 StatID ) c
 
 void FProfilerManager::TrackDefaultStats()
 {
-	// Set the default tracked stats only when loading the first capture or connecting to the first profiler session.
-	TrackStat( 2 );
+	// Find StatId for the game thread.
+	for( auto It = GetProfilerInstancesIterator(); It; ++It )
+	{
+		FProfilerSessionRef ProfilerSession = It.Value();
+		TrackStat( ProfilerSession->GetMetaData()->GetGameThreadStatID() );
+		break;
+	}
 }
 
 /*-----------------------------------------------------------------------------
@@ -693,7 +703,7 @@ void FProfilerManager::TrackDefaultStats()
 
 void FProfilerManager::CloseAllEventGraphTabs()
 {
-	TSharedPtr<class SProfilerWindow> ProfilerWindowPtr = GetProfilerWindow();
+	TSharedPtr<SProfilerWindow> ProfilerWindowPtr = GetProfilerWindow();
 	if( ProfilerWindowPtr.IsValid() )
 	{
 		// Iterate through all profiler sessions.
@@ -707,7 +717,7 @@ void FProfilerManager::CloseAllEventGraphTabs()
 
 void FProfilerManager::DataGraph_OnSelectionChangedForIndex( uint32 FrameStartIndex, uint32 FrameEndIndex )
 {
-	FProfilerScopedLogTime SCI( TEXT("FProfilerManager::DataGraph_OnSelectionChangedForIndex") );
+	PROFILER_SCOPE_LOG_TIME( TEXT( "FProfilerManager::DataGraph_OnSelectionChangedForIndex" ), );
 
 	for( auto It = GetProfilerInstancesIterator(); It; ++It )
 	{
@@ -720,4 +730,13 @@ void FProfilerManager::DataGraph_OnSelectionChangedForIndex( uint32 FrameStartIn
 	}
 }
 
+void FProfilerManager::ProfilerSession_OnAddThreadTime( int32 FrameIndex, const TMap<uint32, float>& ThreadMS, const FProfilerStatMetaDataRef& StatMetaData )
+{
+	TSharedPtr<SProfilerWindow> ProfilerWindowPtr = GetProfilerWindow();
+	if( ProfilerWindowPtr.IsValid() )
+	{
+		ProfilerWindowPtr->ProfilerMiniView->AddThreadTime( FrameIndex, ThreadMS, StatMetaData );
+	}
+}
+
 #undef LOCTEXT_NAMESPACE
\ No newline at end of file
diff --git a/Engine/Source/Developer/Profiler/Private/ProfilerManager.h b/Engine/Source/Developer/Profiler/Private/ProfilerManager.h
--- a/Engine/Source/Developer/Profiler/Private/ProfilerManager.h
+++ b/Engine/Source/Developer/Profiler/Private/ProfilerManager.h
@@ -385,6 +385,7 @@ public:
 	void SetDataCapture( const bool bRequestedDataCaptureState );
 
 	void ProfilerSession_OnCaptureFileProcessed( const FGuid ProfilerInstanceID );
+	void ProfilerSession_OnAddThreadTime( int32 FrameIndex, const TMap<uint32, float>& ThreadMS, const FProfilerStatMetaDataRef& StatMetaData );
 
 	/*-----------------------------------------------------------------------------
 		Event graphs management
diff --git a/Engine/Source/Developer/Profiler/Private/ProfilerPrivatePCH.h b/Engine/Source/Developer/Profiler/Private/ProfilerPrivatePCH.h
--- a/Engine/Source/Developer/Profiler/Private/ProfilerPrivatePCH.h
+++ b/Engine/Source/Developer/Profiler/Private/ProfilerPrivatePCH.h
@@ -41,6 +41,7 @@
 #include "SEventGraph.h"
 #include "SDataGraph.h"
 #include "SProfilerGraphPanel.h"
+#include "SProfilerMiniView.h"
 #include "SProfilerFPSChartPanel.h"
 #include "SProfilerWindow.h"
 #include "SHistogram.h"
diff --git a/Engine/Source/Developer/Profiler/Private/ProfilerSession.cpp b/Engine/Source/Developer/Profiler/Private/ProfilerSession.cpp
--- a/Engine/Source/Developer/Profiler/Private/ProfilerSession.cpp
+++ b/Engine/Source/Developer/Profiler/Private/ProfilerSession.cpp
@@ -111,7 +111,7 @@ FGraphDataSourceRefConst FProfilerSession::CreateGraphDataSource( const uint32 I
 FEventGraphDataRef FProfilerSession::CreateEventGraphData( const uint32 FrameStartIndex, const uint32 FrameEndIndex, const EEventGraphTypes::Type EventGraphType )
 {
 	static FTotalTimeAndCount Current(0.0f, 0);
-	FProfilerScopedLogTime ScopedLog( TEXT("FProfilerSession::CreateEventGraphData"), &Current );
+	PROFILER_SCOPE_LOG_TIME( TEXT( "FProfilerSession::CreateEventGraphData" ), &Current );
 
 	FEventGraphData* EventGraphData = new FEventGraphData();
 
@@ -189,13 +189,13 @@ void FProfilerSession::UpdateAggregatedEventGraphData( const uint32 FrameIndex )
 	if( CompletionSync.GetReference() && !CompletionSync->IsComplete() )
 	{
 		static FTotalTimeAndCount JoinTasksTimeAndCont(0.0f, 0);
-		FProfilerScopedLogTime ScopedLog( TEXT("FProfilerSession::CombineJoinAndContinue"), &JoinTasksTimeAndCont );
+		PROFILER_SCOPE_LOG_TIME( TEXT( "FProfilerSession::CombineJoinAndContinue" ), &JoinTasksTimeAndCont );
 
 		FTaskGraphInterface::Get().WaitUntilTaskCompletes(CompletionSync, ENamedThreads::GameThread);
 	}
 
 	static FTotalTimeAndCount TimeAndCount(0.0f, 0);
-	FProfilerScopedLogTime ScopedLog( TEXT("FProfilerSession::UpdateAggregatedEventGraphData"), &TimeAndCount );
+	PROFILER_SCOPE_LOG_TIME( TEXT( "FProfilerSession::UpdateAggregatedEventGraphData" ), &TimeAndCount );
 
 	// Create a temporary event graph data for the specified frame.
 	EventGraphDataCurrent = MakeShareable( new FEventGraphData( AsShared(), FrameIndex ) );
@@ -249,9 +249,10 @@ bool FProfilerSession::HandleTicker( float DeltaTime )
 		FProfilerDataFrame& CurrentProfilerData = FrameToProfilerDataMapping.FindChecked( FrameIndex );
 
 		static FTotalTimeAndCount Current(0.0f, 0);
-		FProfilerScopedLogTime PSHT( TEXT("FProfilerSession::HandleTicker"), &Current );
+		PROFILER_SCOPE_LOG_TIME( TEXT( "FProfilerSession::HandleTicker" ), &Current );
 
 		const FProfilerStatMetaDataRef MetaData = GetMetaData();
+		TMap<uint32, float> ThreadMS;
 
 		// Preprocess the hierarchical samples for the specified frame.
 		const TMap<uint32, FProfilerCycleGraph>& CycleGraphs = CurrentProfilerData.CycleGraphs;
@@ -299,6 +300,7 @@ bool FProfilerSession::HandleTicker( float DeltaTime )
 					1, 
 					FrameRootSampleIndex 
 				);
+				ThreadMS.FindOrAdd( ThreadID ) = (float)ThreadDurationTimeMS;
 
 				// Recursively add children and parent to the root samples.
 				for( int32 Index = 0; Index < ThreadGraph.Children.Num(); Index++ )
@@ -351,6 +353,9 @@ bool FProfilerSession::HandleTicker( float DeltaTime )
 		// Update aggregated events.
 		UpdateAggregatedEventGraphData( LastFrameIndex );
 
+		// Update mini-view.
+		OnAddThreadTime.ExecuteIfBound( LastFrameIndex, ThreadMS, StatMetaData );
+
 		FrameToProfilerDataMapping.Remove( FrameIndex );
 	}	
 
@@ -433,9 +438,9 @@ void FProfilerSession::PopulateHierarchy_Recurrent
 
 FRawProfilerSession::FRawProfilerSession( const FString& InRawStatsFileFileath )
 : FProfilerSession( EProfilerSessionTypes::OfflineRaw, nullptr, FGuid::NewGuid(), InRawStatsFileFileath.Replace( *FStatConstants::StatsFileRawExtension, TEXT( "" ) ) )
+, CurrentMiniViewFrame( 0 )
 {
 	OnTick = FTickerDelegate::CreateRaw( this, &FRawProfilerSession::HandleTicker );
-	FTicker::GetCoreTicker().AddTicker( OnTick );
 }
 
 FRawProfilerSession::~FRawProfilerSession()
@@ -445,6 +450,38 @@ FRawProfilerSession::~FRawProfilerSession()
 
 bool FRawProfilerSession::HandleTicker( float DeltaTime )
 {
+	StatsThreadStats;
+	Stream;
+
+	enum
+	{
+		MAX_NUM_DATA_PER_TICK = 30
+	};
+	int32 NumDataThisTick = 0;
+
+	// Add the data to the mini-view.
+	for( int32 FrameIndex = CurrentMiniViewFrame; FrameIndex < Stream.FramesInfo.Num(); ++FrameIndex )
+	{
+		const FStatsFrameInfo& StatsFrameInfo = Stream.FramesInfo[FrameIndex];
+		// Convert from cycles to ms.
+		TMap<uint32, float> ThreadMS;
+		for( auto InnerIt = StatsFrameInfo.ThreadCycles.CreateConstIterator(); InnerIt; ++InnerIt )
+		{
+			ThreadMS.Add( InnerIt.Key(), StatMetaData->ConvertCyclesToMS( InnerIt.Value() ) );
+		}
+
+		// Pass the reference to the stats' metadata.
+		// @TODO yrx 2014-04-03 Figure out something better later.
+		OnAddThreadTime.ExecuteIfBound( FrameIndex, ThreadMS, StatMetaData );
+
+		//CurrentMiniViewFrame++;
+		NumDataThisTick++;
+		if( NumDataThisTick > MAX_NUM_DATA_PER_TICK )
+		{
+			break;
+		}
+	}
+
 	return true;
 }
 
@@ -464,7 +501,6 @@ void FRawProfilerSession::PrepareLoading()
 		return;
 	}
 
-	FStatsReadStream Stream;
 	if( !Stream.ReadHeader( *FileReader ) )
 	{
 		UE_LOG( LogStats, Error, TEXT( "Could not open, bad magic: %s" ), *Filepath );
@@ -502,7 +538,7 @@ void FRawProfilerSession::PrepareLoading()
 		}
 		int64 FrameOffset0 = Stream.FramesInfo[0].FrameOffset;
 		FileReader->Seek( FrameOffset0 );
-
+#if	0
 		// Read the raw stats messages.
 		for( int32 FrameIndex = 0; FrameIndex < 30/*Stream.FramesInfo.Num()*/; ++FrameIndex )
 		{
@@ -527,13 +563,19 @@ void FRawProfilerSession::PrepareLoading()
 			// Serialize thread cycles.
 			*FileReader << ThreadCycles;
 		}
+#endif // 0
 	}
 
+	// We have the whole metadata and basic information about the raw stats file, start ticking the profiler session.
+	FTicker::GetCoreTicker().AddTicker( OnTick, 0.25f );
+
+#if	0
 	if( SessionType == EProfilerSessionTypes::OfflineRaw )
 	{
 		// Broadcast that a capture file has been fully processed.
 		OnCaptureFileProcessed.ExecuteIfBound( GetInstanceID() );
 	}
+#endif // 0
 }
 
 void FProfilerStatMetaData::UpdateFromStatsState( const FStatsThreadState& StatsThreadStats )
@@ -640,6 +682,15 @@ void FProfilerStatMetaData::UpdateFromStatsState( const FStatsThreadState& Stats
 				}
 			}
 			ThreadIDtoStatID.Add( ThreadID, StatID );
+
+			if( StatName == NAME_GameThread )
+			{
+				GameThreadID = ThreadID;
+			}
+			else if( StatName == NAME_RenderThread )
+			{
+				RenderThreadIDs.Add( ThreadID );
+			}
 		}
 	}
 
diff --git a/Engine/Source/Developer/Profiler/Private/ProfilerSession.h b/Engine/Source/Developer/Profiler/Private/ProfilerSession.h
--- a/Engine/Source/Developer/Profiler/Private/ProfilerSession.h
+++ b/Engine/Source/Developer/Profiler/Private/ProfilerSession.h
@@ -392,7 +392,7 @@ protected:
 	 */
 	void Update( const FStatMetaData& ClientStatMetaData )
 	{
-		FProfilerScopedLogTime Scoped( TEXT("FProfilerStatMetaData.Update") );
+		PROFILER_SCOPE_LOG_TIME( TEXT( "FProfilerStatMetaData.Update" ), );
 
 		FStatMetaData LocalCopy;
 		{
@@ -405,8 +405,6 @@ protected:
 
 		// Initialize fake stat for Self.
 		const uint32 NoGroupID = 0;
-		FString FF1 = FStatConstants::NAME_ThreadRoot.ToString();
-		FString FF2 = FStatConstants::NAME_ThreadRoot.GetPlainNameString();
 
 		InitializeGroup( NoGroupID, "NoGroup" );
 		InitializeStat( 0, NoGroupID, TEXT( "Self" ), STATTYPE_CycleCounter );
@@ -448,12 +446,12 @@ private:
 	 */
 	void InitializeGroup( const uint32 GroupID, const FString& GroupName ) 
 	{
-		FProfilerGroup** GroupPtr = GroupDescriptions.Find( GroupID );
+		FProfilerGroup* GroupPtr = GroupDescriptions.FindRef( GroupID );
 		if( !GroupPtr )
 		{
-			GroupPtr = &GroupDescriptions.Add( GroupID, new FProfilerGroup( GroupID ) );
+			GroupPtr = GroupDescriptions.Add( GroupID, new FProfilerGroup( GroupID ) );
 		}
-		(*GroupPtr)->Initialize( GroupName );
+		GroupPtr->Initialize( GroupName );
 	}
 
 	/**
@@ -469,24 +467,24 @@ private:
 	void InitializeStat( const uint32 StatID, const uint32 GroupID, const FString& StatName, const EStatType InType, FName StatFName = NAME_None )	
 	{
 		// @TODO yrx 2014-03-24 FindRef!
-		FProfilerGroup** GroupPtr = GroupDescriptions.Find( GroupID );
+		FProfilerGroup* GroupPtr = GroupDescriptions.FindRef( GroupID );
 
-		FProfilerStat** StatPtr = StatDescriptions.Find( StatID );
+		FProfilerStat* StatPtr = StatDescriptions.FindRef( StatID );
 		if( !StatPtr )
 		{
-			StatPtr = &StatDescriptions.Add( StatID, new FProfilerStat( StatID ) );
+			StatPtr = StatDescriptions.Add( StatID, new FProfilerStat( StatID ) );
 
 			if( StatFName != NAME_None )
 			{
-				StatFNameDescriptions.Add( StatFName, *StatPtr );
+				StatFNameDescriptions.Add( StatFName, StatPtr );
 			}
 		}
-		(*StatPtr)->Initialize( StatName, *GroupPtr, InType );
+		StatPtr->Initialize( StatName, GroupPtr, InType );
 
-		if( StatFName == NAME_None && (*GroupPtr)->Name() == TEXT( "Threads" ) )
+		if( StatFName == NAME_None && GroupPtr->Name() == TEXT( "Threads" ) )
 		{
 			// Check if this stat is a thread stat.
-			const uint32 ThreadID = FStatsUtils::ParseThreadID( (*StatPtr)->_Name );
+			const uint32 ThreadID = FStatsUtils::ParseThreadID( StatPtr->Name().ToString() );
 			if( ThreadID != 0 )
 			{
 				const FString* ThreadDesc = GetThreadDescriptions().Find( ThreadID );
@@ -494,15 +492,28 @@ private:
 				{
 					// Replace the stat name with a thread name.
 					const FString UniqueThreadName = FString::Printf( TEXT( "%s [0x%x]" ), **ThreadDesc, ThreadID );
-					(*StatPtr)->_Name = *UniqueThreadName;
+					StatPtr->_Name = *UniqueThreadName;
 					ThreadIDtoStatID.Add( ThreadID, StatID );
-				}
+
+					if( **ThreadDesc == FName( NAME_GameThread ) )
+					{
+						GameThreadID = ThreadID;
+					}
+					else if( **ThreadDesc == FName( NAME_RenderThread ) )
+					{
+						RenderThreadIDs.AddUnique( ThreadID );
+					}
+					else if( ThreadDesc->Contains( TEXT( "RenderingThread" ) ) )
+					{
+						RenderThreadIDs.AddUnique( ThreadID );
+					}
+				}			
 			}
 		}
 
-		if( *GroupPtr != FProfilerGroup::GetDefaultPtr() )
+		if( GroupPtr != FProfilerGroup::GetDefaultPtr() )
 		{
-			(*GroupPtr)->AddStat( *StatPtr );
+			GroupPtr->AddStat( StatPtr );
 		}
 	}
 
@@ -587,6 +598,21 @@ public:
 		return StatDescriptions.Contains( StatID );
 	}
 
+	const uint32 GetGameThreadID() const
+	{
+		return GameThreadID;
+	}
+
+	const TArray<uint32>& GetRenderThreadID() const
+	{
+		return RenderThreadIDs;
+	}
+
+	const uint32 GetGameThreadStatID() const
+	{
+		return ThreadIDtoStatID.FindChecked( GameThreadID );
+	}
+
 protected:
 	/** All initialized stat descriptions, stored as StatID -> FProfilerStat. */
 	TMap<uint32,FProfilerStat*> StatDescriptions;
@@ -608,6 +634,12 @@ protected:
 
 	/** Seconds per CPU cycle. */
 	double SecondsPerCycle;
+
+	/** Game thread id. */
+	uint32 GameThreadID;
+
+	/** Array of all render thread ids. */
+	TArray<uint32> RenderThreadIDs;
 };
 
 /** Holds the aggregated information for the specific stat across all frames that have been captured. */
@@ -934,7 +966,20 @@ public:
 	/** Updates this profiler session. */
 	bool HandleTicker( float DeltaTime );
 
+public:
+	typedef TMap<uint32, float> FThreadTimesMap;
+	DECLARE_DELEGATE_ThreeParams( FAddThreadTimeDelegate, int32 /*FrameIndex*/, const FThreadTimesMap& /*ThreadMS*/, const FProfilerStatMetaDataRef& /*StatMetaData*/ );
+
+	FProfilerSession& SetOnAddThreadTime( const FAddThreadTimeDelegate& InOnAddThreadTime )
+	{
+		OnAddThreadTime = InOnAddThreadTime;
+		return *this;
+	}
+
+protected:
+	FAddThreadTimeDelegate OnAddThreadTime;
 
+public:
 	/**
 	 * The delegate to execute when this profiler session has fully processed a capture file and profiler manager can update all widgets.
 	 *
@@ -1222,6 +1267,10 @@ class FRawProfilerSession : public FProfilerSession
 
 	/** Stats thread state, mostly used to manage the stats metadata. */
 	FStatsThreadState StatsThreadStats;
+	FStatsReadStream Stream;
+
+	/** Index of the last processed data for the mini-view. */
+	int32 CurrentMiniViewFrame;
 
 public:
 	/**
diff --git a/Engine/Source/Developer/Profiler/Private/Widgets/SDataGraph.cpp b/Engine/Source/Developer/Profiler/Private/Widgets/SDataGraph.cpp
--- a/Engine/Source/Developer/Profiler/Private/Widgets/SDataGraph.cpp
+++ b/Engine/Source/Developer/Profiler/Private/Widgets/SDataGraph.cpp
@@ -926,7 +926,6 @@ int32 SDataGraph::OnPaint
 	// Draw all graphs descriptions.
 	float GraphDescPosY = 100.0f;
 
-
 #if DEBUG_PROFILER_PERFORMANCE
 	// Debug text.
 	FSlateDrawElement::MakeText
@@ -970,13 +969,6 @@ int32 SDataGraph::OnPaint
 	);
 	GraphDescPosY += MaxFontCharHeight + 1.0f;
 
-
-#endif // DEBUG_PROFILER_PERFORMANCE
-
-	LayerId++;
-
-
-#if	DEBUG_PROFILER_PERFORMANCE
 	const double CurrentTime = (FPlatformTime::Seconds() - StartTime) * 1000.0f;
 	if( CurrentTime > 1.0f )
 	{
diff --git a/Engine/Source/Developer/Profiler/Private/Widgets/SDataGraph.h b/Engine/Source/Developer/Profiler/Private/Widgets/SDataGraph.h
--- a/Engine/Source/Developer/Profiler/Private/Widgets/SDataGraph.h
+++ b/Engine/Source/Developer/Profiler/Private/Widgets/SDataGraph.h
@@ -415,7 +415,10 @@ protected:
 	/** The current mouse position. */
 	FVector2D MousePosition;
 
+	/** Mouse position during the call on mouse button up. */
 	FVector2D MousePositionOnButtonDown;
+
+	/** Mouse position during the call on mouse button down. */
 	FVector2D MousePositionOnButtonUp;
 
 	/** Accumulated mouse wheel value. */
@@ -485,7 +488,7 @@ protected:
 	/** Index of the frame currently being hovered by the mouse. */
 	int32 HoveredFrameIndex;
 
-	int FrameIndices[2];
+	int32 FrameIndices[2];
 
 	/*-----------------------------------------------------------------------------
 		Time based
diff --git a/Engine/Source/Developer/Profiler/Private/Widgets/SEventGraph.cpp b/Engine/Source/Developer/Profiler/Private/Widgets/SEventGraph.cpp
--- a/Engine/Source/Developer/Profiler/Private/Widgets/SEventGraph.cpp
+++ b/Engine/Source/Developer/Profiler/Private/Widgets/SEventGraph.cpp
@@ -1477,7 +1477,7 @@ static EEventCompareOps::Type EColumnSortModeToEventCompareOp( const EColumnSort
 
 void SEventGraph::SortEvents()
 {
-	FProfilerScopedLogTime Log(TEXT("SEventGraph::SortEvents"));
+	PROFILER_SCOPE_LOG_TIME( TEXT( "SEventGraph::SortEvents" ), );
 
 	if( ColumnBeingSorted != NAME_None )
 	{
@@ -2903,7 +2903,7 @@ TSharedRef<SHorizontalBox> SEventGraph::GetContentForEvent( FEventGraphSamplePtr
 
 void SEventGraph::SetNewEventGraphState( const FEventGraphDataRef AverageEventGraph, const FEventGraphDataRef MaximumEventGraph, bool bInitial )
 {
-	FProfilerScopedLogTime UEG( TEXT("SEventGraph::UpdateEventGraph") );
+	PROFILER_SCOPE_LOG_TIME( TEXT( "SEventGraph::UpdateEventGraph" ), );
 
 	// Store current operation.
 	SaveCurrentEventGraphState();
@@ -3177,7 +3177,7 @@ bool SEventGraph::ClearHistory_CanExecute() const
 
 void SEventGraph::FilterOutByProperty_Execute( const FEventGraphSamplePtr EventPtr, const FName PropertyName, const bool bReset )
 {
-	FProfilerScopedLogTime FilterOutByProperty(TEXT("SEventGraph::FilterOutByProperty_Execute"));
+	PROFILER_SCOPE_LOG_TIME( TEXT( "SEventGraph::FilterOutByProperty_Execute" ), );
 
 	// Store current operation.
 	SaveCurrentEventGraphState();
@@ -3201,7 +3201,7 @@ bool SEventGraph::FilterOutByProperty_CanExecute( const FEventGraphSamplePtr Eve
 
 void SEventGraph::CullByProperty_Execute( const FEventGraphSamplePtr EventPtr, const FName PropertyName, const bool bReset )
 {
-	FProfilerScopedLogTime CullByProperty(TEXT("SEventGraph::CullByProperty_Execute"));
+	PROFILER_SCOPE_LOG_TIME( TEXT( "SEventGraph::CullByProperty_Execute" ), );
 
 	// Store current operation.
 	SaveCurrentEventGraphState();
diff --git a/Engine/Source/Developer/Profiler/Private/Widgets/SProfilerGraphPanel.cpp b/Engine/Source/Developer/Profiler/Private/Widgets/SProfilerGraphPanel.cpp
--- a/Engine/Source/Developer/Profiler/Private/Widgets/SProfilerGraphPanel.cpp
+++ b/Engine/Source/Developer/Profiler/Private/Widgets/SProfilerGraphPanel.cpp
@@ -35,7 +35,7 @@ void SProfilerGraphPanel::Construct( const FArguments& InArgs )
 			.FillHeight(1.0f)
 			[
 				SAssignNew(DataGraph,SDataGraph)
-				.OnGraphOffsetChanged( this, &SProfilerGraphPanel::DataGraph_OnGraphOffsetChanged )
+				.OnGraphOffsetChanged( this, &SProfilerGraphPanel::OnDataGraphGraphOffsetChanged )
 				.OnViewModeChanged( this, &SProfilerGraphPanel::DataGraph_OnViewModeChanged )
 			]
 
@@ -74,6 +74,8 @@ void SProfilerGraphPanel::ScrollBar_OnUserScrolled( float ScrollOffset )
  
 	GraphOffset = FMath::Trunc( ScrollOffset * NumDataPoints );
 	DataGraph->ScrollTo( GraphOffset );
+
+	FProfilerManager::Get()->GetProfilerWindow()->ProfilerMiniView->OnSelectionBoxMoved( GraphOffset, GraphOffset + NumVisiblePoints );
 }
 
 void SProfilerGraphPanel::ProfilerManager_OnTrackedStatChanged( const FTrackedStat& TrackedStat, bool bIsTracked )
@@ -93,12 +95,18 @@ void SProfilerGraphPanel::ProfilerManager_OnSessionInstancesUpdated()
 
 }
 
-void SProfilerGraphPanel::DataGraph_OnGraphOffsetChanged( int InGraphOffset )
+void SProfilerGraphPanel::OnDataGraphGraphOffsetChanged( int32 InGraphOffset )
 {
 	GraphOffset = InGraphOffset;
-	const float ThumbSizeFraction = FMath::Min( (float)NumVisiblePoints / (float)NumDataPoints, 1.0f );
-	const float OffsetFraction = (float)GraphOffset / (float)NumDataPoints;
-	ScrollBar->SetState( OffsetFraction, ThumbSizeFraction );
+	SetScrollBarState();
+	FProfilerManager::Get()->GetProfilerWindow()->ProfilerMiniView->OnSelectionBoxMoved( InGraphOffset, InGraphOffset + NumVisiblePoints );
+}
+
+void SProfilerGraphPanel::OnMiniViewSelectionBoxChanged( int32 FrameStart, int32 FrameEnd )
+{
+	GraphOffset = FrameStart;
+	SetScrollBarState();
+	DataGraph->ScrollTo( GraphOffset );
 }
 
 void SProfilerGraphPanel::DataGraph_OnViewModeChanged( EDataGraphViewModes::Type ViewMode )
@@ -111,13 +119,8 @@ void SProfilerGraphPanel::UpdateInternals()
 	NumVisiblePoints = DataGraph->GetNumVisiblePoints();
 	NumDataPoints = DataGraph->GetNumDataPoints();
 
-	// Note that the maximum offset is 1.0-ThumbSizeFraction.
-	// If the user can view 1/3 of the items in a single page, the maximum offset will be ~0.667f
-	const float ThumbSizeFraction = FMath::Min( (float)NumVisiblePoints / (float)NumDataPoints, 1.0f );
-	const float OffsetFraction = (float)GraphOffset / (float)NumDataPoints;
-		
-	const float Distance = ScrollBar->DistanceFromBottom();
-	ScrollBar->SetState( OffsetFraction, ThumbSizeFraction );
+	SetScrollBarState();
+	FProfilerManager::Get()->GetProfilerWindow()->ProfilerMiniView->OnSelectionBoxMoved( GraphOffset, GraphOffset + NumVisiblePoints );
 	 
 	if( FProfilerManager::Get()->IsLivePreview() )
 	{
@@ -126,4 +129,13 @@ void SProfilerGraphPanel::UpdateInternals()
 	}
 }
 
+void SProfilerGraphPanel::SetScrollBarState()
+{
+	// Note that the maximum offset is 1.0-ThumbSizeFraction.
+	// If the user can view 1/3 of the items in a single page, the maximum offset will be ~0.667f
+	const float ThumbSizeFraction = FMath::Min( (float)NumVisiblePoints / (float)NumDataPoints, 1.0f );
+	const float OffsetFraction = (float)GraphOffset / (float)NumDataPoints;
+	ScrollBar->SetState( OffsetFraction, ThumbSizeFraction );
+}
+
 #undef LOCTEXT_NAMESPACE
diff --git a/Engine/Source/Developer/Profiler/Private/Widgets/SProfilerGraphPanel.h b/Engine/Source/Developer/Profiler/Private/Widgets/SProfilerGraphPanel.h
--- a/Engine/Source/Developer/Profiler/Private/Widgets/SProfilerGraphPanel.h
+++ b/Engine/Source/Developer/Profiler/Private/Widgets/SProfilerGraphPanel.h
@@ -45,8 +45,16 @@ protected:
 	 */
 	void ScrollBar_OnUserScrolled( float ScrollOffset );
 
-	/** Called when the frame offset has been changed from the data graph widget. */
-	void DataGraph_OnGraphOffsetChanged( int InFrameOffset );
+	/** Called when the frame offset has been changed in the data graph widget. */
+	void OnDataGraphGraphOffsetChanged( int32 InFrameOffset );
+
+public:
+	/** Called when the selection box has been changed in the profiler mini-view widget. */
+	void OnMiniViewSelectionBoxChanged( int32 FrameStart, int32 FrameEnd );
+
+protected:
+	/** Sets state of the scroll bar. */
+	void SetScrollBarState();
 
 	/** Called when the data graph view mode has changed from the data graph widget. */
 	void DataGraph_OnViewModeChanged( EDataGraphViewModes::Type ViewMode );
diff --git a/Engine/Source/Developer/Profiler/Private/Widgets/SProfilerMiniView.cpp b/Engine/Source/Developer/Profiler/Private/Widgets/SProfilerMiniView.cpp
new file mode 100644
--- /dev/null
+++ b/Engine/Source/Developer/Profiler/Private/Widgets/SProfilerMiniView.cpp
@@ -0,0 +1,701 @@
+// Copyright 1998-2014 Epic Games, Inc. All Rights Reserved.
+
+#include "ProfilerPrivatePCH.h"
+
+#define LOCTEXT_NAMESPACE "SProfilerMiniView"
+
+// SProfilerMiniView
+#undef LOCTEXT_NAMESPACE
+
+SProfilerMiniView::SProfilerMiniView()
+{
+	Reset();
+}
+
+SProfilerMiniView::~SProfilerMiniView()
+{
+
+}
+
+void SProfilerMiniView::Construct( const FArguments& InArgs )
+{
+	BindCommands();
+}
+
+void SProfilerMiniView::Reset()
+{
+	MaxFrameTime = 0.0f;
+	AllFrames.Empty();
+	RecentlyAddedFrames.Empty();
+	StatMetadata = nullptr;
+
+	FMemory::MemZero( PaintStateMemory );
+	PaintState = nullptr;
+
+	MousePosition = FVector2D::ZeroVector;
+	MousePositionOnButtonDown = FVector2D::ZeroVector;
+
+	SelectionBoxFrameStart = 0;
+	SelectionBoxFrameEnd = 0;
+
+	HoveredFrameIndex = 0;
+	DistanceDragged = 0.0f;
+	NumPixelsPerSample = 0;
+	NumPixelsPerFrame = 0.0f;
+
+	bIsLeftMousePressed = false;
+	bIsRightMousePressed = false;
+	bCanBeStartDragged = false;
+	bCanBeEndDragged = false;
+	bAllowSelectionBoxZooming = false;
+
+	CursorType = EMiniviewCursor::Default;
+}
+
+
+void SProfilerMiniView::Tick( const FGeometry& AllottedGeometry, const double InCurrentTime, const float InDeltaTime )
+{
+	SCompoundWidget::Tick( AllottedGeometry, InCurrentTime, InDeltaTime );
+
+	if( ThisGeometry != AllottedGeometry )
+	{
+		// Refresh.
+		MaxFrameTime = 0.0f;
+		bUpdateData = true;
+	}
+
+	ThisGeometry = AllottedGeometry;
+
+	if( RecentlyAddedFrames.Num() > 0 )
+	{
+		bUpdateData = true;
+	}
+
+	if( IsReady() )
+	{
+		NumPixelsPerFrame = (float)AllottedGeometry.Size.X / (float)AllFrames.Num();
+	}
+
+	if( ShouldUpdateData() )
+	{
+		ProcessData();
+		bUpdateData = false;
+	}
+}
+
+int32 SProfilerMiniView::OnPaint( const FGeometry& AllottedGeometry, const FSlateRect& MyClippingRect, FSlateWindowElementList& OutDrawElements, int32 LayerId, const FWidgetStyle& InWidgetStyle, bool bParentEnabled ) const
+{
+//	SCOPE_LOG_TIME_FUNC();
+
+	// Rendering info.
+	const bool bEnabled = ShouldBeEnabled( bParentEnabled );
+	const ESlateDrawEffect::Type DrawEffects = bEnabled ? ESlateDrawEffect::None : ESlateDrawEffect::DisabledEffect;
+	const FSlateBrush* MiniViewArea = FEditorStyle::GetBrush( "Profiler.LineGraphArea" );
+	const FSlateBrush* WhiteBrush = FEditorStyle::GetBrush( "WhiteTexture" );
+
+	PaintState = new((void*)PaintStateMemory) FSlateOnPaintState( AllottedGeometry, MyClippingRect, OutDrawElements, LayerId, InWidgetStyle, DrawEffects );
+
+	const float MiniViewSizeX = AllottedGeometry.Size.X;
+	const float MiniViewSizeY = AllottedGeometry.Size.Y;
+
+	const TSharedRef< FSlateFontMeasure > FontMeasureService = FSlateApplication::Get().GetRenderer()->GetFontMeasureService();
+	FSlateFontInfo SummaryFont( FPaths::EngineContentDir() / TEXT( "Slate/Fonts/Roboto-Regular.ttf" ), 8 );
+	const float MaxFontCharHeight = FontMeasureService->Measure( TEXT( "!" ), SummaryFont ).Y;
+
+	// Draw background.
+	FSlateDrawElement::MakeBox
+	(
+		OutDrawElements,
+		LayerId,
+		AllottedGeometry.ToPaintGeometry( FVector2D( 0, 0 ), FVector2D( MiniViewSizeX, AllottedGeometry.Size.Y ) ),
+		MiniViewArea,
+		MyClippingRect,
+		DrawEffects,
+		MiniViewArea->GetTint( InWidgetStyle ) * InWidgetStyle.GetColorAndOpacityTint()
+	);
+	LayerId++;
+
+	// Draw all samples.
+	if( IsReady() )
+	{
+		static const FSlateColorBrush SolidWhiteBrush = FSlateColorBrush( FColorList::White );
+		const FColor GameThreadColor = FColorList::Red;
+		const FColor RenderThreadColor = FColorList::Blue;
+		const FColor OtherThreadsColor = FColorList::Grey;
+
+		const float SampleScaleY = MiniViewSizeY / MaxFrameTime;
+
+		float CurrentSamplePosX = 0;
+		float NextSamplePosX = GetNumPixelsPerSample();
+		for( const FMiniViewSample& MiniViewSample : MiniViewSamples )
+		{
+			const float AllSizeY = FMath::TruncFloat( MiniViewSample.TotalThreadTime*SampleScaleY );
+			const float GTSizeY = FMath::TruncFloat( MiniViewSample.GameThreadTime*SampleScaleY );
+			const float RTSizeY = FMath::TruncFloat( MiniViewSample.RenderThreadTime*SampleScaleY );
+			const float OtherSizeY = AllSizeY - GTSizeY - RTSizeY;
+
+			const float DestSamplePosX0 = FMath::TruncFloat( CurrentSamplePosX );
+			const float DestSamplePosX1 = FMath::TruncFloat( NextSamplePosX );
+			const float DestSampleSizeX = DestSamplePosX1 - DestSamplePosX0;
+
+			// The game thread on the bottom.
+			FSlateDrawElement::MakeBox
+			(
+				OutDrawElements,
+				LayerId,
+				AllottedGeometry.ToPaintGeometry( FVector2D( DestSamplePosX0, MiniViewSizeY - GTSizeY ), FVector2D( DestSampleSizeX, GTSizeY ) ),
+				&SolidWhiteBrush,
+				MyClippingRect,
+				DrawEffects,
+				GameThreadColor
+			);
+
+			// Next the render thread.
+			FSlateDrawElement::MakeBox
+			(
+				OutDrawElements,
+				LayerId,
+				AllottedGeometry.ToPaintGeometry( FVector2D( DestSamplePosX0, MiniViewSizeY - GTSizeY - RTSizeY ), FVector2D( DestSampleSizeX, RTSizeY ) ),
+				&SolidWhiteBrush,
+				MyClippingRect,
+				DrawEffects,
+				RenderThreadColor
+			);
+
+			// Disabled for now.
+#if	0
+			// Finally the others threads as one sample.
+			FSlateDrawElement::MakeBox
+			(
+				OutDrawElements,
+				LayerId,
+				AllottedGeometry.ToPaintGeometry( FVector2D( DestSamplePosX, MiniViewSizeY - AllSizeY ), FVector2D( DestSampleSizeX, OtherSizeY ) ),
+				&SolidWhiteBrush,
+				MyClippingRect,
+				DrawEffects,
+				OtherThreadsColor
+			);
+#endif // 0
+
+			CurrentSamplePosX = NextSamplePosX;
+			NextSamplePosX += GetNumPixelsPerSample();
+		}
+
+		// Draw the selection box
+		LayerId++;
+		
+		const int32 MaxFrameIndex = AllFrames.Num() - 1;
+		
+		const int32 SelectionBoxX0 = FMath::Trunc( FrameIndexToPosition( SelectionBoxFrameStart ) );
+		const int32 SelectionBoxX1 = FMath::Trunc( FrameIndexToPosition( SelectionBoxFrameEnd + 1 ) );
+
+		if( SelectionBoxFrameStart > 0 )
+		{
+			FSlateDrawElement::MakeBox
+			(
+				OutDrawElements,
+				LayerId,
+				AllottedGeometry.ToPaintGeometry( FVector2D( 0.0f, 0.0f ), FVector2D( (float)SelectionBoxX0, AllottedGeometry.Size.Y ) ),
+				&SolidWhiteBrush,
+				MyClippingRect,
+				DrawEffects,
+				FColorList::Grey.WithAlpha( 192 )
+			);
+		}
+	
+		if( SelectionBoxFrameEnd < MaxFrameIndex )
+		{
+			FSlateDrawElement::MakeBox
+			(
+				OutDrawElements,
+				LayerId,
+				AllottedGeometry.ToPaintGeometry( FVector2D( SelectionBoxX1, 0.0f ), FVector2D( MiniViewSizeX - SelectionBoxX1, AllottedGeometry.Size.Y ) ),
+				&SolidWhiteBrush,
+				MyClippingRect,
+				DrawEffects,
+				FColorList::Grey.WithAlpha( 192 )
+			);
+		}
+
+		// Draw the filler, to hide the difference between window's width and mini-view samples' width.
+		const float FillerSizeX = MiniViewSizeX - CurrentSamplePosX;
+		if( FillerSizeX > 0.0f )
+		{
+			const float FillerPosX0 = MiniViewSizeX - FillerSizeX;
+			FSlateDrawElement::MakeBox
+			(
+				OutDrawElements,
+				LayerId,
+				AllottedGeometry.ToPaintGeometry( FVector2D( FillerPosX0, 0.0f ), FVector2D( FillerSizeX + 1.0f, AllottedGeometry.Size.Y ) ),
+				&SolidWhiteBrush,
+				MyClippingRect,
+				DrawEffects,
+				// @TODO yrx 2014-04-09 How to get this color from Slate?
+				FColor(96,96,96)
+			);
+		}
+
+		// Border of the selection box.
+		LayerId++;
+		FSlateDrawElement::MakeBox
+		(
+			OutDrawElements,
+			LayerId,
+			AllottedGeometry.ToPaintGeometry( FVector2D( SelectionBoxX0, 0.0f ), FVector2D( SelectionBoxX1 - SelectionBoxX0, AllottedGeometry.Size.Y ) ),
+			FEditorStyle::GetBrush( "PlainBorder" ),
+			MyClippingRect,
+			DrawEffects,
+			FColorList::Green
+		);
+
+		// Draw the current mouse position with the highlighted sample.
+		LayerId++;
+
+		// Draw the basic information about threads and data range.
+		const float MarkerPosX = 4.0f;
+		const float MarkerPosY = 4.0f;
+
+		// Upper left
+		DrawText( TEXT( "Rendering thread" ), SummaryFont, FVector2D( MarkerPosX, MarkerPosY ), RenderThreadColor, FColor::Black, FVector2D( 1.0f, 1.0f ) );
+		// Down left
+		DrawText( TEXT( "Game thread" ), SummaryFont, FVector2D( MarkerPosX, MiniViewSizeY - MarkerPosY - MaxFontCharHeight ), GameThreadColor, FColor::Black, FVector2D( 1.0f, 1.0f ) );
+
+		// Upper right
+		const FString ThreadTimeMax = FString::Printf( TEXT( "%5.2f MS" ), MaxFrameTime );
+		const float ThreadTimeMaxSizeX = FontMeasureService->Measure( ThreadTimeMax, SummaryFont ).X;
+		DrawText( ThreadTimeMax, SummaryFont, FVector2D( MiniViewSizeX - ThreadTimeMaxSizeX - MarkerPosX, MarkerPosY ), FColor::White, FColor::Black, FVector2D( 1.0f, 1.0f ) );
+
+		// Down right
+		const FString ThreadTimeMin = TEXT( "0.0 MS" );
+		const float ThreadTimeMinSizeX = FontMeasureService->Measure( ThreadTimeMin, SummaryFont ).X;
+		DrawText( ThreadTimeMin, SummaryFont, FVector2D( MiniViewSizeX - ThreadTimeMinSizeX - MarkerPosX, MiniViewSizeY - MarkerPosY - MaxFontCharHeight ), FColor::White, FColor::Black, FVector2D( 1.0f, 1.0f ) );
+	
+	}
+
+#if 0/*DEBUG_PROFILER_PERFORMANCE*/
+	
+	LayerId++;
+	// Draw debug information.
+	float GraphDescPosY = 0;
+
+	// Debug text.
+	FSlateDrawElement::MakeText
+	(
+		OutDrawElements,
+		LayerId,
+		AllottedGeometry.ToOffsetPaintGeometry( FVector2D( 16.0f, GraphDescPosY ) ),
+		FString::Printf( TEXT( "SelectionBox FrameStart=%4i, FrameEnd=%4i (%3i) Hovered=%4i, NumPixelsPerSample=%2.1f" ), SelectionBoxFrameStart, SelectionBoxFrameEnd, SelectionBoxFrameEnd - SelectionBoxFrameStart, HoveredFrameIndex, GetNumPixelsPerSample() ),
+		SummaryFont,
+		MyClippingRect,
+		DrawEffects,
+		FLinearColor::White
+	);
+	GraphDescPosY += MaxFontCharHeight + 1.0f;
+
+	FSlateDrawElement::MakeText
+	(
+		OutDrawElements,
+		LayerId,
+		AllottedGeometry.ToOffsetPaintGeometry( FVector2D( 16.0f, GraphDescPosY ) ),
+		FString::Printf( TEXT( "AllFrames=%4i, MiniViewSamples=%3i, MaxFrameTime=%2.0f DistanceDragged=%2.0f" ), AllFrames.Num(), MiniViewSamples.Num(), MaxFrameTime, DistanceDragged ),
+		SummaryFont,
+		MyClippingRect,
+		DrawEffects,
+		FLinearColor::White
+	);
+	GraphDescPosY += MaxFontCharHeight + 1.0f;
+
+	FSlateDrawElement::MakeText
+	(
+		OutDrawElements,
+		LayerId,
+		AllottedGeometry.ToOffsetPaintGeometry( FVector2D( 16.0f, GraphDescPosY ) ),
+		FString::Printf( TEXT( "bCanBeStartD=%1i, bCanBeEndD=%1i" ), (int32)bCanBeStartDragged, (int32)bCanBeEndDragged ),
+		SummaryFont,
+		MyClippingRect,
+		DrawEffects,
+		FLinearColor::White
+	);
+	GraphDescPosY += MaxFontCharHeight + 1.0f;
+
+#endif // DEBUG_PROFILER_PERFORMANCE
+
+	return SCompoundWidget::OnPaint( AllottedGeometry, MyClippingRect, OutDrawElements, LayerId, InWidgetStyle, bParentEnabled && IsEnabled() );
+}
+
+
+void SProfilerMiniView::DrawText( const FString& Text, const FSlateFontInfo& FontInfo, FVector2D Position, const FColor& TextColor, const FColor& ShadowColor, FVector2D ShadowOffset ) const
+{
+	if( ShadowOffset.Size() > 0.0f )
+	{
+		FSlateDrawElement::MakeText
+		(
+			PaintState->OutDrawElements,
+			PaintState->LayerId,
+			PaintState->AllottedGeometry.ToOffsetPaintGeometry( Position + ShadowOffset ),
+			Text,
+			FontInfo,
+			PaintState->MyClippingRect,
+			PaintState->DrawEffects,
+			ShadowColor
+		);
+	}
+
+	FSlateDrawElement::MakeText
+	(
+		PaintState->OutDrawElements,
+		++PaintState->LayerId,
+		PaintState->AllottedGeometry.ToOffsetPaintGeometry( Position ),
+		Text,
+		FontInfo,
+		PaintState->MyClippingRect,
+		PaintState->DrawEffects,
+		TextColor
+	);
+}
+
+
+FReply SProfilerMiniView::OnMouseButtonDown( const FGeometry& MyGeometry, const FPointerEvent& MouseEvent )
+{
+	FReply Reply = FReply::Unhandled();
+
+	if( IsReady() )
+	{
+		MousePositionOnButtonDown = MyGeometry.AbsoluteToLocal( MouseEvent.GetScreenSpacePosition() );
+
+		if( MouseEvent.GetEffectingButton() == EKeys::LeftMouseButton )
+		{
+			bIsLeftMousePressed = true;
+
+			if( bCanBeStartDragged )
+			{
+				DistanceDragged = FrameIndexToPosition( SelectionBoxFrameStart );
+			}
+			else if( bCanBeEndDragged )
+			{
+				DistanceDragged = FrameIndexToPosition( SelectionBoxFrameEnd );
+			}
+			else
+			{
+				// Clicked outside the selection box, so move the selection box to that position.
+				DistanceDragged = MousePositionOnButtonDown.X;
+			}
+
+			if( bCanBeStartDragged || bCanBeEndDragged )
+			{
+				// Capture mouse, so we can move outside this widget.
+				Reply = FReply::Handled().CaptureMouse( SharedThis( this ) );
+			}
+		}
+		else if( MouseEvent.GetEffectingButton() == EKeys::RightMouseButton )
+		{
+			bIsRightMousePressed = true;
+		}
+	}
+
+	return Reply;
+}
+
+FReply SProfilerMiniView::OnMouseButtonUp( const FGeometry& MyGeometry, const FPointerEvent& MouseEvent )
+{
+	FReply Reply = FReply::Unhandled();
+	
+	if( IsReady() )
+	{
+		const FVector2D MousePositionOnButtonUp = MyGeometry.AbsoluteToLocal( MouseEvent.GetScreenSpacePosition() );
+		const bool bIsValidForMouseClick = MousePositionOnButtonUp.Equals( MousePositionOnButtonDown, MOUSE_SNAP_DISTANCE );
+
+		if( MouseEvent.GetEffectingButton() == EKeys::LeftMouseButton )
+		{
+			if( bIsLeftMousePressed )
+			{
+				if( !bCanBeStartDragged && !bCanBeEndDragged )
+				{
+					MoveSelectionBox( PositionToFrameIndex( DistanceDragged ) );
+				}
+				else if( bCanBeStartDragged || bCanBeEndDragged )
+				{
+					// Release the mouse, we are no longer dragging.
+					Reply = FReply::Handled().ReleaseMouseCapture();
+				}		
+			}
+
+			bIsLeftMousePressed = false;
+		}
+		else if( MouseEvent.GetEffectingButton() == EKeys::RightMouseButton )
+		{
+			if( bIsRightMousePressed )
+			{
+				if( bIsValidForMouseClick )
+				{
+					ShowContextMenu( MouseEvent.GetScreenSpacePosition() );
+					Reply = FReply::Handled();
+				}
+			}
+			bIsRightMousePressed = false;
+		}
+	}
+
+	return Reply;
+}
+
+FReply SProfilerMiniView::OnMouseMove( const FGeometry& MyGeometry, const FPointerEvent& MouseEvent )
+{
+//	SCOPE_LOG_TIME_FUNC();
+
+	FReply Reply = FReply::Unhandled();
+
+	if( IsReady() )
+	{
+		MousePosition = MyGeometry.AbsoluteToLocal( MouseEvent.GetScreenSpacePosition() );
+		HoveredFrameIndex = PositionToFrameIndex( MousePosition.X );
+
+		const float CursorPosXDelta = MouseEvent.GetCursorDelta().X;
+
+		if( MouseEvent.IsMouseButtonDown( EKeys::LeftMouseButton ) )
+		{
+			if( HasMouseCapture() && !MouseEvent.GetCursorDelta().IsZero() )
+			{
+				DistanceDragged += CursorPosXDelta;
+
+				const int32 MouseFrameIndex = PositionToFrameIndex( DistanceDragged );
+
+				if( !bAllowSelectionBoxZooming )
+				{
+					const int32 SelectionBoxSize = SelectionBoxFrameEnd - SelectionBoxFrameStart;
+					if( bCanBeStartDragged )
+					{				
+						SelectionBoxFrameStart = MouseFrameIndex;
+						SelectionBoxFrameStart = FMath::Clamp( SelectionBoxFrameStart, 0, FMath::Max( AllFrames.Num() - SelectionBoxSize - 1, 0 ) );
+						
+						SelectionBoxFrameEnd = SelectionBoxFrameStart + SelectionBoxSize;
+					}
+					else if( bCanBeEndDragged )
+					{
+						SelectionBoxFrameEnd = MouseFrameIndex;
+						SelectionBoxFrameEnd = FMath::Clamp( SelectionBoxFrameEnd, SelectionBoxSize - 1, AllFrames.Num() - -1 );
+
+						SelectionBoxFrameStart = SelectionBoxFrameEnd - SelectionBoxSize;
+					}
+				}
+				else
+				{
+					
+				}
+
+				// Inform other widgets that we have moved the selection box.
+				SelectionBoxChangedEvent.Broadcast( SelectionBoxFrameStart, SelectionBoxFrameEnd );
+
+				Reply = FReply::Handled();
+			}
+		}
+		else
+		{
+			const float StartEdgeDistance = FrameIndexToPosition(SelectionBoxFrameStart) - MousePosition.X;
+			const float EndEdgeDistance = MousePosition.X - FrameIndexToPosition(SelectionBoxFrameEnd);
+
+			bCanBeStartDragged = StartEdgeDistance < MOUSE_SNAP_DISTANCE && StartEdgeDistance > 0.0f;
+			bCanBeEndDragged = EndEdgeDistance < MOUSE_SNAP_DISTANCE && EndEdgeDistance > 0.0f;
+
+			if( StartEdgeDistance <= 0.0f && EndEdgeDistance <= 0.0f )
+			{
+				// Drag the whole selection box.
+				bCanBeStartDragged = true;
+				bCanBeEndDragged = true;
+				CursorType = EMiniviewCursor::Hand;
+			}
+			else if( bCanBeStartDragged != bCanBeEndDragged )
+			{
+				CursorType = EMiniviewCursor::Arrow;
+			}
+			else
+			{
+				CursorType = EMiniviewCursor::Default;
+			}
+		}
+	}
+
+	return Reply;
+}
+
+void SProfilerMiniView::OnMouseEnter( const FGeometry& MyGeometry, const FPointerEvent& MouseEvent )
+{
+
+}
+
+void SProfilerMiniView::OnMouseLeave( const FPointerEvent& MouseEvent )
+{
+	if( !HasMouseCapture() )
+	{
+		bIsLeftMousePressed = false;
+		bIsRightMousePressed = false;
+
+		bCanBeStartDragged = false;
+		bCanBeEndDragged = false;
+
+		CursorType = EMiniviewCursor::Default;
+	}
+}
+
+FReply SProfilerMiniView::OnMouseWheel( const FGeometry& MyGeometry, const FPointerEvent& MouseEvent )
+{
+	FReply Reply = FReply::Unhandled();
+	return Reply;
+}
+
+FReply SProfilerMiniView::OnMouseButtonDoubleClick( const FGeometry& MyGeometry, const FPointerEvent& MouseEvent )
+{
+	FReply Reply = FReply::Unhandled();
+	return Reply;
+}
+
+
+FCursorReply SProfilerMiniView::OnCursorQuery( const FGeometry& MyGeometry, const FPointerEvent& CursorEvent ) const
+{
+	FCursorReply CursorReply = FCursorReply::Unhandled();
+	if( CursorType == EMiniviewCursor::Arrow )
+	{
+		CursorReply = FCursorReply::Cursor( EMouseCursor::ResizeLeftRight );
+	}
+	else if( CursorType == EMiniviewCursor::Hand )
+	{
+		CursorReply = FCursorReply::Cursor( EMouseCursor::GrabHand );
+	}
+	return CursorReply;
+}
+
+void SProfilerMiniView::ShowContextMenu( const FVector2D& ScreenSpacePosition )
+{
+
+}
+
+void SProfilerMiniView::BindCommands()
+{
+
+}
+
+
+
+void SProfilerMiniView::AddThreadTime( int32 InFrameIndex, const TMap<uint32, float>& InThreadMS, const FProfilerStatMetaDataRef& InStatMetaData )
+{
+	FFrameThreadTimes FrameThreadTimes;
+	FrameThreadTimes.FrameNumber = InFrameIndex;
+	FrameThreadTimes.ThreadTimes = InThreadMS;
+	RecentlyAddedFrames.Add( FrameThreadTimes );
+
+	StatMetadata = InStatMetaData;
+}
+
+void SProfilerMiniView::OnSelectionBoxMoved( int32 FrameStart, int32 FrameEnd )
+{
+	const int32 MaxFrameIndex = FMath::Max( AllFrames.Num() - 1, 0 );
+	SelectionBoxFrameStart = FMath::Clamp( FrameStart, 0, MaxFrameIndex );
+	SelectionBoxFrameEnd = FMath::Clamp( FrameEnd, 0, MaxFrameIndex );
+}
+
+void SProfilerMiniView::MoveSelectionBox( int32 FrameIndex )
+{
+	const int32 SelectionBoxSize = SelectionBoxFrameEnd - SelectionBoxFrameStart;
+	const int32 SelectionBoxHalfSize = SelectionBoxSize / 2;
+	const int32 CenterFrameIndex = FMath::Clamp( FrameIndex - SelectionBoxHalfSize, 0, AllFrames.Num() - 1 - SelectionBoxSize );
+
+	// Inform other widgets that we have moved the selection box.
+	SelectionBoxFrameStart = CenterFrameIndex;
+	SelectionBoxFrameEnd = CenterFrameIndex + SelectionBoxSize;
+	SelectionBoxChangedEvent.Broadcast( SelectionBoxFrameStart, SelectionBoxFrameEnd );
+}
+
+const int32 SProfilerMiniView::PositionToFrameIndex( const float InPositionX ) const
+{
+	int32 FrameIndex = 0;
+
+	if( IsReady() )
+	{
+		const int32 NumAllFrames = AllFrames.Num();
+		const int32 NumMiniViewSamples = MiniViewSamples.Num();
+
+		const float ScaleRatio = (float)NumAllFrames / (float)NumMiniViewSamples;
+		const float MouseSampleIndex = InPositionX / GetNumPixelsPerSample();
+
+		FrameIndex = FMath::Trunc( MouseSampleIndex * ScaleRatio );
+		FrameIndex = FMath::Clamp( FrameIndex, 0, NumAllFrames - 1 );
+	}
+
+	return FrameIndex;
+}
+
+
+void SProfilerMiniView::ProcessData()
+{
+//	SCOPE_LOG_TIME_FUNC();
+
+	AllFrames.Append( RecentlyAddedFrames );
+	RecentlyAddedFrames.Reset();
+
+	if( !IsReady() )
+	{
+		return;
+	}
+
+	UpdateNumPixelsPerSample();
+
+	const int32 NumMiniViewSamples = FMath::Trunc( ThisGeometry.Size.X / GetNumPixelsPerSample() );
+	MiniViewSamples.Reset( NumMiniViewSamples );
+
+	MiniViewSamples.AddZeroed( NumMiniViewSamples );
+	const int32 NumAllFrames = AllFrames.Num();
+
+	const uint32 GameThreadID = StatMetadata->GetGameThreadID();
+	const TArray<uint32>& RenderThreadIDs = StatMetadata->GetRenderThreadID();
+
+	// Prepare data to fit into the window.
+	{
+		// Aggregate.
+		const float ScaleRatio = (float)NumMiniViewSamples / (float)NumAllFrames;
+
+		float CurrentSample = 0.0f;
+		for( int32 FrameIndex = 0; FrameIndex < NumAllFrames; ++FrameIndex )
+		{
+			const FFrameThreadTimes& FrameThreadTimes = AllFrames[FrameIndex];
+			const int32 SampleIndex = FMath::Trunc( CurrentSample );
+
+			FMiniViewSample& Dest = MiniViewSamples[SampleIndex];
+			Dest.AddFrameAndFindMax( FrameThreadTimes );
+			
+			CurrentSample += ScaleRatio;
+		}
+
+		for( FMiniViewSample& It : MiniViewSamples )
+		{
+			It.CalculateMaxThreadTime( GameThreadID, RenderThreadIDs );
+		}
+	}
+#if	0
+	else
+	{
+		// Replicate.
+		const float ScaleRatio = (float)NumAllFrames / (float)WindowNumSamples;
+
+		float CurrentFrameIndex = 0.0f;
+		for( int32 SampleIndex = 0; SampleIndex < WindowNumSamples; ++SampleIndex )
+		{
+			const int FrameIndex = FMath::Trunc( CurrentFrameIndex );
+			const FFrameThreadTimes& FrameThreadTimes = AllFrames[FrameIndex];
+
+			FMiniViewThreadTimes& Dest = WindowFrames[SampleIndex];
+			Dest.Aggregate( AllFrames[FrameIndex] );
+			Dest.CalculateMaxThreadTime( GameThreadID, RenderThreadIDs );
+
+			CurrentFrameIndex += ScaleRatio;
+		}
+	}
+#endif // 0
+
+	// Calculate max thread time, used to scaling the displayed samples.
+	for( const FMiniViewSample& MiniViewSample : MiniViewSamples )
+	{
+		MaxFrameTime = FMath::Max( MaxFrameTime, MiniViewSample.TotalThreadTime );
+	}
+
+	MaxFrameTime = FMath::Clamp( MaxFrameTime, 0.0f, (float)MAX_VISIBLE_THREADTIME );
+}
diff --git a/Engine/Source/Developer/Profiler/Private/Widgets/SProfilerMiniView.h b/Engine/Source/Developer/Profiler/Private/Widgets/SProfilerMiniView.h
new file mode 100644
--- /dev/null
+++ b/Engine/Source/Developer/Profiler/Private/Widgets/SProfilerMiniView.h
@@ -0,0 +1,475 @@
+// Copyright 1998-2014 Epic Games, Inc. All Rights Reserved.
+
+#pragma once
+
+/*-----------------------------------------------------------------------------
+	Basic structures
+-----------------------------------------------------------------------------*/
+
+struct FFrameThreadTimes
+{
+	/** Number of the frame. */
+	int FrameNumber;
+	/** Thread times for the frame. */
+	TMap<uint32, float> ThreadTimes;
+};
+
+struct FMiniViewSample
+{
+	/** Frames for one mini-view frame, at least 1. */
+	TArray<int32> Frames;
+
+	/** Aggregates thread times for the Frames. */
+	TMap<uint32,float> ThreadTimes;
+
+	/**
+	 *	Accumulative time for all threads. 
+	 *	At this moment only for the game and render threads.
+	 */
+	float TotalThreadTime;
+
+	/** Accumulative time for the game thread. */
+	float GameThreadTime;
+
+	/** Accumulative time for the render thread. */
+	float RenderThreadTime;
+
+	FMiniViewSample()
+		: TotalThreadTime(0.0f)
+		, GameThreadTime(0.0f)
+		, RenderThreadTime(0.0f)
+	{}
+
+	void AddFrameAndAccumulate( const FFrameThreadTimes& FrameThreadTimes )
+	{
+		new(Frames) int32( FrameThreadTimes.FrameNumber );
+		for( auto It = FrameThreadTimes.ThreadTimes.CreateConstIterator(); It; ++It )
+		{
+			ThreadTimes.FindOrAdd( It.Key() ) += It.Value();
+		}
+	}
+
+	void AddFrameAndFindMax( const FFrameThreadTimes& FrameThreadTimes )
+	{
+		new(Frames) int32( FrameThreadTimes.FrameNumber );
+		for( auto It = FrameThreadTimes.ThreadTimes.CreateConstIterator(); It; ++It )
+		{
+			float& ThreadMS = ThreadTimes.FindOrAdd( It.Key() );
+			ThreadMS = FMath::Max( ThreadMS, It.Value() );
+		}
+	}
+
+	void CalculateTotalThreadTime( const uint32 GameThreadID, const TArray<uint32>& RenderThreadIDs )
+	{
+		TotalThreadTime = 0.0f;
+		GameThreadTime = 0.0f;
+		RenderThreadTime = 0.0f;
+		for( auto It = ThreadTimes.CreateConstIterator(); It; ++It )
+		{
+			const float CurrentTimeMS = It.Value();		
+
+			if( It.Key() == GameThreadID )
+			{
+				TotalThreadTime += CurrentTimeMS;
+				GameThreadTime += CurrentTimeMS;
+			}
+			else if( RenderThreadIDs.Contains(It.Key()) )
+			{
+				TotalThreadTime += CurrentTimeMS;
+				RenderThreadTime += CurrentTimeMS;
+			}
+		}
+
+		TotalThreadTime /= (float)Frames.Num();
+		GameThreadTime /= (float)Frames.Num();
+		RenderThreadTime /= (float)Frames.Num();
+	}
+
+	void CalculateMaxThreadTime( const uint32 GameThreadID, const TArray<uint32>& RenderThreadIDs )
+	{
+		TotalThreadTime = 0.0f;
+		GameThreadTime = 0.0f;
+		RenderThreadTime = 0.0f;
+
+		float CurrentGameThreadTime = 0.0f;
+		float CurrentRenderThreadTime = 0.0f;
+
+		for( auto It = ThreadTimes.CreateConstIterator(); It; ++It )
+		{
+			const float CurrentTimeMS = It.Value();
+
+			if( It.Key() == GameThreadID )
+			{
+				CurrentGameThreadTime = CurrentTimeMS;
+			}
+			else if( RenderThreadIDs.Contains( It.Key() ) )
+			{
+				CurrentRenderThreadTime = CurrentTimeMS;
+			}
+
+			const float CurrentTotalTile = CurrentGameThreadTime + CurrentRenderThreadTime;
+			if( CurrentTotalTile > TotalThreadTime )
+			{
+				TotalThreadTime = CurrentTotalTile;
+				GameThreadTime = CurrentRenderThreadTime;
+				RenderThreadTime = CurrentGameThreadTime;
+			}
+		}
+	}
+};
+
+/*-----------------------------------------------------------------------------
+	Declarations
+-----------------------------------------------------------------------------*/
+
+/** Widget used to present thread data in the mini-view. */
+class SProfilerMiniView : public SCompoundWidget
+{
+	enum 
+	{
+		/** Minimum width of the one rendered sample. */
+		MIN_NUM_PIXELS_PER_SAMPLE = 4,
+
+		/** Number of pixels. */
+		MOUSE_SNAP_DISTANCE = 4,
+
+		/**
+		 *	Maximum total thread time that will be visible on the mini-view.
+		 *	This is enough to see performance issues on the mini-view.
+		 *	Everything else will be clamped to that value.
+		 */
+		MAX_VISIBLE_THREADTIME = 150,
+	};
+
+	struct EMiniviewCursor
+	{
+		enum Type
+		{
+			Default,
+			Arrow,
+			Hand,
+		};
+	};
+
+	/** Holds current state provided by OnPaint function, used to simplify drawing. */
+	struct FSlateOnPaintState : public FNoncopyable
+	{
+		FSlateOnPaintState( const FGeometry& InAllottedGeometry, const FSlateRect& InMyClippingRect, FSlateWindowElementList& InOutDrawElements, int32& InLayerId, const FWidgetStyle& InWidgetStyle, ESlateDrawEffect::Type InDrawEffects )
+			: AllottedGeometry( InAllottedGeometry )
+			, MyClippingRect( InMyClippingRect )
+			, WidgetStyle( InWidgetStyle )
+			, OutDrawElements( InOutDrawElements )
+			, LayerId( InLayerId )
+			, DrawEffects( InDrawEffects )
+		{}
+
+#if	0
+		/** Set operation. */
+		void Set( const FGeometry& InAllottedGeometry, const FSlateRect& InMyClippingRect, FSlateWindowElementList& InOutDrawElements, int32& InLayerId, const FWidgetStyle& InWidgetStyle, ESlateDrawEffect::Type InDrawEffects )
+		{
+			_AllottedGeometry = &InAllottedGeometry;
+			_MyClippingRect = &InMyClippingRect;
+			_WidgetStyle = &InWidgetStyle;
+			_OutDrawElements = &InOutDrawElements;
+			_LayerId = &InLayerId;
+			_DrawEffects = InDrawEffects;
+		}
+#endif // 0
+
+		/** Accessors. */
+		const FGeometry& AllottedGeometry; 
+		const FSlateRect& MyClippingRect;
+		const FWidgetStyle& WidgetStyle;
+		 
+		FSlateWindowElementList& OutDrawElements;
+		int32& LayerId;
+		const ESlateDrawEffect::Type DrawEffects;
+	};
+
+public:
+	SProfilerMiniView();
+	~SProfilerMiniView();
+
+	SLATE_BEGIN_ARGS( SProfilerMiniView )
+		{}
+	SLATE_END_ARGS()
+
+	/**
+	* Construct this widget
+	*
+	* @param	InArgs	The declaration data for this widget
+	*/
+	void Construct( const FArguments& InArgs );
+
+	/** Resets internal widget's data to the default one. */
+	void Reset();
+
+	/**
+	* Ticks this widget.  Override in derived classes, but always call the parent implementation.
+	*
+	* @param  AllottedGeometry The space allotted for this widget
+	* @param  InCurrentTime  Current absolute real time
+	* @param  InDeltaTime  Real time passed since last tick
+	*/
+	virtual void Tick( const FGeometry& AllottedGeometry, const double InCurrentTime, const float InDeltaTime ) OVERRIDE;
+
+	/**
+	* The widget should respond by populating the OutDrawElements array with FDrawElements
+	* that represent it and any of its children.
+	*
+	* @param AllottedGeometry  The FGeometry that describes an area in which the widget should appear.
+	* @param MyClippingRect    The clipping rectangle allocated for this widget and its children.
+	* @param OutDrawElements   A list of FDrawElements to populate with the output.
+	* @param LayerId           The Layer onto which this widget should be rendered.
+	* @param InColorAndOpacity ColorAverage and Opacity to be applied to all the descendants of the widget being painted
+	* @param bParentEnabled	True if the parent of this widget is enabled.
+	*
+	* @return The maximum layer ID attained by this widget or any of its children.
+	*/
+	virtual int32 OnPaint( const FGeometry& AllottedGeometry, const FSlateRect& MyClippingRect, FSlateWindowElementList& OutDrawElements, int32 LayerId, const FWidgetStyle& InWidgetStyle, bool bParentEnabled ) const OVERRIDE;
+
+	/**
+	* The system calls this method to notify the widget that a mouse button was pressed within it. This event is bubbled.
+	*
+	* @param MyGeometry The Geometry of the widget receiving the event
+	* @param MouseEvent Information about the input event
+	*
+	* @return Whether the event was handled along with possible requests for the system to take action.
+	*/
+	virtual FReply OnMouseButtonDown( const FGeometry& MyGeometry, const FPointerEvent& MouseEvent ) OVERRIDE;
+
+	/**
+	* The system calls this method to notify the widget that a mouse button was release within it. This event is bubbled.
+	*
+	* @param MyGeometry The Geometry of the widget receiving the event
+	* @param MouseEvent Information about the input event
+	*
+	* @return Whether the event was handled along with possible requests for the system to take action.
+	*/
+	virtual FReply OnMouseButtonUp( const FGeometry& MyGeometry, const FPointerEvent& MouseEvent ) OVERRIDE;
+
+	/**
+	* The system calls this method to notify the widget that a mouse moved within it. This event is bubbled.
+	*
+	* @param MyGeometry The Geometry of the widget receiving the event
+	* @param MouseEvent Information about the input event
+	*
+	* @return Whether the event was handled along with possible requests for the system to take action.
+	*/
+	virtual FReply OnMouseMove( const FGeometry& MyGeometry, const FPointerEvent& MouseEvent ) OVERRIDE;
+
+	/**
+	* The system will use this event to notify a widget that the cursor has entered it. This event is NOT bubbled.
+	*
+	* @param MyGeometry The Geometry of the widget receiving the event
+	* @param MouseEvent Information about the input event
+	*/
+	virtual void OnMouseEnter( const FGeometry& MyGeometry, const FPointerEvent& MouseEvent ) OVERRIDE;
+
+	/**
+	* The system will use this event to notify a widget that the cursor has left it. This event is NOT bubbled.
+	*
+	* @param MouseEvent Information about the input event
+	*/
+	virtual void OnMouseLeave( const FPointerEvent& MouseEvent ) OVERRIDE;
+
+	/**
+	* Called when the mouse wheel is spun. This event is bubbled.
+	*
+	* @param  MouseEvent  Mouse event
+	*
+	* @return  Returns whether the event was handled, along with other possible actions
+	*/
+	virtual FReply OnMouseWheel( const FGeometry& MyGeometry, const FPointerEvent& MouseEvent ) OVERRIDE;
+
+	/**
+	* Called when a mouse button is double clicked.  Override this in derived classes.
+	*
+	* @param  InMyGeometry  Widget geometry
+	* @param  InMouseEvent  Mouse button event
+	*
+	* @return  Returns whether the event was handled, along with other possible actions
+	*/
+	virtual FReply OnMouseButtonDoubleClick( const FGeometry& MyGeometry, const FPointerEvent& MouseEvent ) OVERRIDE;
+
+	/**
+	* Called when the system wants to know which cursor to display for this Widget.  This event is bubbled.
+	*
+	* @return  The cursor requested (can be None.)
+	*/
+	virtual FCursorReply OnCursorQuery( const FGeometry& MyGeometry, const FPointerEvent& CursorEvent ) const OVERRIDE;
+
+	void ShowContextMenu( const FVector2D& ScreenSpacePosition );
+
+	/**	Binds UI commands to delegates. */
+	void BindCommands();
+
+	/** For the specified frame ands a thread data. */
+	void AddThreadTime( int32 InFrameIndex, const TMap<uint32, float>& InThreadMS, const FProfilerStatMetaDataRef& InStatMetaData );
+
+	/**
+	 *	Moves the selection box to the location specified by FrameStart and FrameEnd. 
+	 *	Called by external widgets.
+	 */
+	void OnSelectionBoxMoved( int32 FrameStart, int32 FrameEnd );
+
+	void DrawText( const FString& Text, const FSlateFontInfo& FontInfo, FVector2D Position, const FColor& TextColor, const FColor& ShadowColor, FVector2D ShadowOffset ) const;
+
+protected:
+	/**
+	 *	Moves the selection box to the location specified by FrameIndex.
+	 *	Centers it at the same time.
+	 */
+	void MoveSelectionBox( int32 FrameIndex );
+
+public:
+	/** The event to execute when the selection box has been changed. */
+	DECLARE_EVENT_TwoParams( SProfilerMiniView, FSelectionBoxChangedEvent, int32 /*FrameStart*/, int32 /*FrameEnd*/ );
+	FSelectionBoxChangedEvent& OnSelectionBoxChanged()
+	{
+		return SelectionBoxChangedEvent;
+	}
+	
+protected:
+	/** The event to execute when the selection box has been changed. */
+	FSelectionBoxChangedEvent SelectionBoxChangedEvent;
+
+	/**
+	 * @return Index of the frame currently being hovered by the mouse
+	 */
+	int32 GetHoveredFrameIndex() const
+	{
+		return HoveredFrameIndex;
+	}
+
+	/**
+	* @return calculates frame index based on the specified mouse position.
+	*/
+	const int32 PositionToFrameIndex( const float InPositionX ) const;
+
+	/**
+	* @return calculates position based on the specified frame index.
+	*/
+	const float FrameIndexToPosition( const int32 FrameIndex ) const
+	{
+		return FrameIndex*NumPixelsPerFrame;
+	}
+
+protected:
+	void ProcessData();
+
+	void UpdateNumPixelsPerSample()
+	{
+		NumPixelsPerSample = FMath::Max( (float)MIN_NUM_PIXELS_PER_SAMPLE, ThisGeometry.Size.X / AllFrames.Num() );
+	}
+
+	const float GetNumPixelsPerSample() const
+	{
+		return NumPixelsPerSample;
+	}
+
+	const int32 FindMiniViewSampleIndex( const int32 FrameIndex ) const
+	{
+		// @TODO yrx 2014-04-07 Naive implementation, consider using binary search.
+		for( int32 SampleIndex = 0; SampleIndex < MiniViewSamples.Num(); ++SampleIndex )
+		{
+			const FMiniViewSample& MiniViewSample = MiniViewSamples[SampleIndex];
+			for( const int32& Frame : MiniViewSample.Frames )
+			{
+				if( FrameIndex == Frame )
+				{
+					return SampleIndex;
+				}
+			}
+		}
+
+		return 0;
+	}
+
+
+	/**
+	 * @return True, if the widget is ready to use, also means that contains at least one frame of the thread data.
+	 */
+	bool IsReady() const
+	{
+		return AllFrames.Num() > 0;
+	}
+
+	bool ShouldUpdateData()
+	{
+		return bUpdateData;
+	}
+
+protected:
+
+	/*-----------------------------------------------------------------------------
+		Data variables
+	-----------------------------------------------------------------------------*/
+
+	/** Processed data used to render the mini-view samples. */
+	TArray<FMiniViewSample> MiniViewSamples;
+
+	/**
+	 *	All mini-view frames history
+	 *	60fps*60s*60m=216k samples for one hour, 
+	 *	For 16 threads it's about 16mb of data.
+	 */
+	TArray<FFrameThreadTimes> AllFrames;
+
+	/** Recently added thread times, not processed yet. */
+	TArray<FFrameThreadTimes> RecentlyAddedFrames;
+	TLockFreePointerList<FFrameThreadTimes> RecentlyAddedFramesLockFree;
+
+	/** Maximum total mini-view frame time seen so far. */
+	float MaxFrameTime;
+
+	/** Shared pointer to the stats' metadata. */
+	FProfilerStatMetaDataPtr StatMetadata;
+
+	/*-----------------------------------------------------------------------------
+		UI variables
+	-----------------------------------------------------------------------------*/
+
+	FGeometry ThisGeometry;
+
+	/** Current Slate OnPaint state. */
+	uint8 PaintStateMemory[sizeof FSlateOnPaintState];
+	mutable FSlateOnPaintState*	PaintState;
+
+	/** The current mouse position. */
+	FVector2D MousePosition;
+
+	/** Mouse position during the call on mouse button down. */
+	FVector2D MousePositionOnButtonDown;
+
+	/** Frame start and frame end for the selection box. */
+	int32 SelectionBoxFrameStart;
+	int32 SelectionBoxFrameEnd;
+
+	/** Index of the frame currently being hovered by the mouse. */
+	int32 HoveredFrameIndex;
+
+	/** Distance dragged. */
+	float DistanceDragged;
+
+	/** Width of the one rendered sample, min value is 4. */
+	float NumPixelsPerSample;
+
+	/** Number of pixels for one frame. */
+	float NumPixelsPerFrame;
+
+	bool bIsLeftMousePressed;
+	bool bIsRightMousePressed;
+
+	bool bCanBeStartDragged;
+	bool bCanBeEndDragged;
+
+	/** Whether to allow zooming through the selection box. */
+	bool bAllowSelectionBoxZooming;
+
+	/** Whether to updated data, set to true during window resize or if new data has been added. */
+	bool bUpdateData;
+
+	/** Cursor type. */
+	EMiniviewCursor::Type CursorType;
+};
\ No newline at end of file
diff --git a/Engine/Source/Developer/Profiler/Private/Widgets/SProfilerWindow.cpp b/Engine/Source/Developer/Profiler/Private/Widgets/SProfilerWindow.cpp
--- a/Engine/Source/Developer/Profiler/Private/Widgets/SProfilerWindow.cpp
+++ b/Engine/Source/Developer/Profiler/Private/Widgets/SProfilerWindow.cpp
@@ -126,6 +126,34 @@ void SProfilerWindow::Construct( const FArguments& InArgs, const ISessionManager
 					.Size( FVector2D( 2.0f, 2.0f ) )
 				]
 
+				/** Profiler Mini-view. */
+				+ SVerticalBox::Slot()
+				.AutoHeight()
+				[
+					SNew(SBox)
+					.HeightOverride( 48.0f )
+					[
+						SNew( SHorizontalBox )
+						.IsEnabled( this, &SProfilerWindow::IsProfilerEnabled )
+
+						+ SHorizontalBox::Slot()
+						.FillWidth( 1.0f )
+						.Padding( 0.0f )
+						.HAlign( HAlign_Fill )
+						.VAlign( VAlign_Fill )
+						[
+							SAssignNew( ProfilerMiniView, SProfilerMiniView )
+						]
+					]
+				]
+
+				+ SVerticalBox::Slot()
+				.AutoHeight()
+				[
+					SNew( SSpacer )
+					.Size( FVector2D( 2.0f, 2.0f ) )
+				]
+
 				+SVerticalBox::Slot()
 				.FillHeight( 1.0f )
 				.Padding(0.0f, 8.0f, 0.0f, 0.0f)
@@ -270,8 +298,11 @@ void SProfilerWindow::Construct( const FArguments& InArgs, const ISessionManager
 			.Expose( OverlaySettingsSlot )
 		];
 
-	GraphPanel->GetMainDataGraph()->OnSelectionChangedForIndex().AddSP( FProfilerManager::Get().Get(), &FProfilerManager::DataGraph_OnSelectionChangedForIndex );
+	GraphPanel->GetMainDataGraph()->OnSelectionChangedForIndex().AddSP( FProfilerManager::Get().ToSharedRef(), &FProfilerManager::DataGraph_OnSelectionChangedForIndex );
+
+	ProfilerMiniView->OnSelectionBoxChanged().AddSP( GraphPanel.ToSharedRef(), &SProfilerGraphPanel::OnMiniViewSelectionBoxChanged );
 }
+
 void SProfilerWindow::ManageEventGraphTab( const FGuid ProfilerInstanceID, const bool bCreateFakeTab, const FString TabName )
 {
 	// TODO: Add support for multiple instances.
diff --git a/Engine/Source/Developer/Profiler/Private/Widgets/SProfilerWindow.h b/Engine/Source/Developer/Profiler/Private/Widgets/SProfilerWindow.h
--- a/Engine/Source/Developer/Profiler/Private/Widgets/SProfilerWindow.h
+++ b/Engine/Source/Developer/Profiler/Private/Widgets/SProfilerWindow.h
@@ -115,4 +115,7 @@ private:
 
 	/** Active event graphs, one event graph for each profiler instance, stored as FGuid -> SEventGraph. */
 	TMap< FGuid, TSharedRef<class SEventGraph> > ActiveEventGraphs;
+
+public:
+	TSharedPtr<SProfilerMiniView> ProfilerMiniView;
 };
\ No newline at end of file
diff --git a/Engine/Source/Developer/Profiler/Public/ProfilerCommon.h b/Engine/Source/Developer/Profiler/Public/ProfilerCommon.h
--- a/Engine/Source/Developer/Profiler/Public/ProfilerCommon.h
+++ b/Engine/Source/Developer/Profiler/Public/ProfilerCommon.h
@@ -58,50 +58,11 @@ class FProfilerAggregatedStat;
 
 #define DEBUG_PROFILER_PERFORMANCE 0
 
-typedef TKeyValuePair<double,uint32> FTotalTimeAndCount;
-
-#if	DEBUG_PROFILER_PERFORMANCE == 1
-
-struct FProfilerScopedLogTime
-{
-	FProfilerScopedLogTime( const TCHAR* InStatName, FTotalTimeAndCount* InGlobal = nullptr )
-		: StartTime( FPlatformTime::Seconds() )
-		, StatName( InStatName )
-		, Global( InGlobal )
-	{}
-
-	~FProfilerScopedLogTime()
-	{
-		const double ScopedTime = FPlatformTime::Seconds() - StartTime;
-		if( Global )
-		{
-			Global->Key += ScopedTime;
-			Global->Value ++;
-
-			const double Average = Global->Key / (double)Global->Value;
-			UE_LOG( Profiler, Log, TEXT("%32s - %6.3f ms - Total %6.2f s / %5u / %6.3f ms"), *StatName, ScopedTime*1000.0f, Global->Key, Global->Value, Average*1000.0f );
-		}
-		else
-		{
-			UE_LOG( Profiler, Log, TEXT("%32s - %6.3f ms - No Total"), *StatName, ScopedTime*1000.0f );
-		}
-	}
-
-protected:
-	const double StartTime;
-	const FString StatName;
-	FTotalTimeAndCount* Global;
-};
-
+#if DEBUG_PROFILER_PERFORMANCE==1
+#define PROFILER_SCOPE_LOG_TIME(arg0,arg1) SCOPE_LOG_TIME(arg0,arg1)
 #else
-
-struct FProfilerScopedLogTime
-{
-	FProfilerScopedLogTime( const TCHAR* InStatName, FTotalTimeAndCount* InGlobal = nullptr )
-	{}
-};
-
-#endif // DEBUG_PROFILER_PERFORMANCE
+#define PROFILER_SCOPE_LOG_TIME(arg0,arg1)
+#endif
 
 /** Time spent on graph drawing. */
 DECLARE_CYCLE_STAT_EXTERN( TEXT("DataGraphOnPaint"),	STAT_DG_OnPaint,			STATGROUP_Profiler, );
