#UE4.Profiler
- Fixed a crash on profiler exit

diff --git a/Engine/Source/Developer/Profiler/Private/ProfilerManager.cpp b/Engine/Source/Developer/Profiler/Private/ProfilerManager.cpp
--- a/Engine/Source/Developer/Profiler/Private/ProfilerManager.cpp
+++ b/Engine/Source/Developer/Profiler/Private/ProfilerManager.cpp
@@ -547,7 +547,6 @@ bool FProfilerManager::UntrackStat( const uint32 StatID )
 void FProfilerManager::ClearStatsAndInstances()
 {
 	CloseAllEventGraphTabs();
-	GetProfilerWindow()->ProfilerMiniView->Reset();
 
 	ProfilerType = EProfilerSessionTypes::InvalidOrMax;
 	SetDataPreview( false );
@@ -712,6 +711,8 @@ void FProfilerManager::CloseAllEventGraphTabs()
 			FProfilerSessionRef ProfilerSession = It.Value();
 			ProfilerWindowPtr->ManageEventGraphTab( ProfilerSession->GetInstanceID(), false, TEXT("") );
 		}
+
+		ProfilerWindowPtr->ProfilerMiniView->Reset();
 	}
 }
 
