Fixed Clang compile error in SProfilerMiniView.h

diff --git a/Engine/Source/Developer/Profiler/Private/Widgets/SProfilerMiniView.h b/Engine/Source/Developer/Profiler/Private/Widgets/SProfilerMiniView.h
--- a/Engine/Source/Developer/Profiler/Private/Widgets/SProfilerMiniView.h
+++ b/Engine/Source/Developer/Profiler/Private/Widgets/SProfilerMiniView.h
@@ -433,7 +433,7 @@ protected:
 	FGeometry ThisGeometry;
 
 	/** Current Slate OnPaint state. */
-	uint8 PaintStateMemory[sizeof FSlateOnPaintState];
+	uint8 PaintStateMemory[sizeof(FSlateOnPaintState)];
 	mutable FSlateOnPaintState*	PaintState;
 
 	/** The current mouse position. */
