fixed up incorrect module dependencies in Profiler

diff --git a/Engine/Source/Developer/Profiler/Private/ProfilerPrivatePCH.h b/Engine/Source/Developer/Profiler/Private/ProfilerPrivatePCH.h
--- a/Engine/Source/Developer/Profiler/Private/ProfilerPrivatePCH.h
+++ b/Engine/Source/Developer/Profiler/Private/ProfilerPrivatePCH.h
@@ -18,6 +18,8 @@
 /*-----------------------------------------------------------------------------
 	Dependencies
 -----------------------------------------------------------------------------*/
+
+#include "Messaging.h"
 #include "ProfilerClient.h"
 
 /*-----------------------------------------------------------------------------
diff --git a/Engine/Source/Developer/Profiler/Profiler.Build.cs b/Engine/Source/Developer/Profiler/Profiler.Build.cs
--- a/Engine/Source/Developer/Profiler/Profiler.Build.cs
+++ b/Engine/Source/Developer/Profiler/Profiler.Build.cs
@@ -8,8 +8,7 @@ public class Profiler : ModuleRules
 	{
 		PrivateIncludePaths.AddRange
 		(
-			new string[] 
-			{
+			new string[] {
 				"Developer/Profiler/Private",
 				"Developer/Profiler/Private/Widgets",
 			}
@@ -17,25 +16,28 @@ public class Profiler : ModuleRules
 
 		PublicDependencyModuleNames.AddRange
 		(
-			 new string[]
-			{
+			new string[] {
 				"Core",
                 "InputCore",
 				"RHI",
 				"RenderCore",
 				"Slate",
                 "EditorStyle",
-				"Messaging",
-				"SessionServices",
 				"ProfilerClient",
 				"DesktopPlatform",
 			}
 		);
 
+		PrivateIncludePathModuleNames.AddRange(
+			new string[] {
+				"Messaging",
+				"SessionServices",
+			}
+		);
+
 		PublicIncludePaths.AddRange
 		(
-			new string[]
-			{
+			new string[] {
 				"Developer/Profiler/Public",
 				"Developer/Profiler/Public/Interfaces",
 			}
diff --git a/Engine/Source/Developer/Profiler/Public/Profiler.h b/Engine/Source/Developer/Profiler/Public/Profiler.h
--- a/Engine/Source/Developer/Profiler/Public/Profiler.h
+++ b/Engine/Source/Developer/Profiler/Public/Profiler.h
@@ -6,9 +6,9 @@
 
 #pragma once
 
-/*-----------------------------------------------------------------------------
-	Dependencies
------------------------------------------------------------------------------*/
+
+/* Dependencies
+ *****************************************************************************/
 
 #include "Core.h"
 #include "ChunkedArray.h"
@@ -17,8 +17,6 @@
 #include "SessionServices.h"
 #include "Slate.h"
 #include "EditorStyle.h"
-#include "Messaging.h"
-
 #include "IFilter.h"
 #include "TextFilter.h"
 #include "FilterCollection.h"
@@ -26,15 +24,15 @@
 #include "IProfilerServiceManager.h"
 #include "IProfilerClient.h"
 
-/*-----------------------------------------------------------------------------
-	Common
------------------------------------------------------------------------------*/
+
+/* Common
+ *****************************************************************************/
+
 #include "ProfilerCommon.h"
 
-/*-----------------------------------------------------------------------------
-	Interfaces
------------------------------------------------------------------------------*/
+
+/* Interfaces
+ *****************************************************************************/
 
 #include "IProfilerManager.h"
 #include "IProfilerModule.h"
-
diff --git a/Engine/Source/Developer/ProfilerClient/Private/ProfilerClientPrivatePCH.h b/Engine/Source/Developer/ProfilerClient/Private/ProfilerClientPrivatePCH.h
--- a/Engine/Source/Developer/ProfilerClient/Private/ProfilerClientPrivatePCH.h
+++ b/Engine/Source/Developer/ProfilerClient/Private/ProfilerClientPrivatePCH.h
@@ -13,11 +13,13 @@
 /* Dependencies
  *****************************************************************************/
 
+#include "Messaging.h"
 #include "ProfilerMessages.h"
 #include "Ticker.h"
 #include "StatsData.h"
 #include "StatsFile.h"
 
+
 /* Private includes
  *****************************************************************************/
 
diff --git a/Engine/Source/Developer/ProfilerClient/ProfilerClient.Build.cs b/Engine/Source/Developer/ProfilerClient/ProfilerClient.Build.cs
--- a/Engine/Source/Developer/ProfilerClient/ProfilerClient.Build.cs
+++ b/Engine/Source/Developer/ProfilerClient/ProfilerClient.Build.cs
@@ -12,8 +12,6 @@ namespace UnrealBuildTool.Rules
 				new string[]
 				{
 					"Core",
-					"Messaging",
-					"SessionServices",
 				}
 			);
 
@@ -24,6 +22,13 @@ namespace UnrealBuildTool.Rules
 					"ProfilerMessages",
 				}
 			);
+
+			PrivateIncludePathModuleNames.AddRange(
+				new string[] {
+					"Messaging",
+					"SessionServices",
+				}
+			);
 		}
 	}
 }
diff --git a/Engine/Source/Developer/ProfilerClient/Public/ProfilerClient.h b/Engine/Source/Developer/ProfilerClient/Public/ProfilerClient.h
--- a/Engine/Source/Developer/ProfilerClient/Public/ProfilerClient.h
+++ b/Engine/Source/Developer/ProfilerClient/Public/ProfilerClient.h
@@ -12,10 +12,10 @@
 
 #include "Core.h"
 #include "ModuleManager.h"
-#include "Messaging.h"
 #include "SessionServices.h"
 #include "ProfilerService.h"
 
+
 /* Interfaces
  *****************************************************************************/
 
diff --git a/Engine/Source/Developer/ProfilerService/Private/ProfilerServicePrivatePCH.h b/Engine/Source/Developer/ProfilerService/Private/ProfilerServicePrivatePCH.h
--- a/Engine/Source/Developer/ProfilerService/Private/ProfilerServicePrivatePCH.h
+++ b/Engine/Source/Developer/ProfilerService/Private/ProfilerServicePrivatePCH.h
@@ -11,6 +11,7 @@
  *****************************************************************************/
 
 #include "Engine.h"
+#include "Messaging.h"
 #include "ProfilerMessages.h"
 #include "Ticker.h"
 #include "StatsData.h"
diff --git a/Engine/Source/Developer/ProfilerService/ProfilerService.Build.cs b/Engine/Source/Developer/ProfilerService/ProfilerService.Build.cs
--- a/Engine/Source/Developer/ProfilerService/ProfilerService.Build.cs
+++ b/Engine/Source/Developer/ProfilerService/ProfilerService.Build.cs
@@ -10,7 +10,6 @@ namespace UnrealBuildTool.Rules
 				new string[]
 				{
 					"Core",
-					"Messaging"
 				}
 			);
 
@@ -19,10 +18,15 @@ namespace UnrealBuildTool.Rules
 				{
 					"CoreUObject",
 					"Engine",
-					"Messaging",
 					"ProfilerMessages",
 				}
 			);
+
+			PrivateIncludePathModuleNames.AddRange(
+				new string[] {
+					"Messaging",
+				}
+			);
 		}
 	}
 }
\ No newline at end of file
diff --git a/Engine/Source/Developer/ProfilerService/Public/ProfilerService.h b/Engine/Source/Developer/ProfilerService/Public/ProfilerService.h
--- a/Engine/Source/Developer/ProfilerService/Public/ProfilerService.h
+++ b/Engine/Source/Developer/ProfilerService/Public/ProfilerService.h
@@ -7,17 +7,15 @@
 #pragma once
 
 
-/*----------------------------------------------------------------------------
-	Dependencies
-----------------------------------------------------------------------------*/
+/* Dependencies
+ *****************************************************************************/
+
 #include "Core.h"
 #include "ModuleManager.h"
 
-#include "Messaging.h"
 
-/*----------------------------------------------------------------------------
-	Interfaces
-----------------------------------------------------------------------------*/
+/* Interfaces
+ *****************************************************************************/
+
 #include "IProfilerServiceManager.h"
 #include "IProfilerServiceModule.h"
-
