#UE4.Profiler
- Added vertical scroll bar to the profiler graph panel, not used yet

diff --git a/Engine/Source/Developer/Profiler/Private/Widgets/SProfilerGraphPanel.cpp b/Engine/Source/Developer/Profiler/Private/Widgets/SProfilerGraphPanel.cpp
--- a/Engine/Source/Developer/Profiler/Private/Widgets/SProfilerGraphPanel.cpp
+++ b/Engine/Source/Developer/Profiler/Private/Widgets/SProfilerGraphPanel.cpp
@@ -2,6 +2,7 @@
 
 #include "ProfilerPrivatePCH.h"
 
+// @TODO yrx 2014-04-18 Move SProfilerWindow later
 #define LOCTEXT_NAMESPACE "SProfilerGraphPanel"
 
 SProfilerGraphPanel::SProfilerGraphPanel()
@@ -29,30 +30,48 @@ void SProfilerGraphPanel::Construct( const FArguments& InArgs )
 		.BorderImage( FEditorStyle::GetBrush("ToolPanel.GroupBorder") )
 		.Padding( 2.0f )
 		[
-			SNew( SVerticalBox )
+			SNew( SHorizontalBox )
 
-			+SVerticalBox::Slot()
-			.FillHeight(1.0f)
+			+SHorizontalBox::Slot()
+			.FillWidth( 1.0f )
 			[
-				SAssignNew(DataGraph,SDataGraph)
-				.OnGraphOffsetChanged( this, &SProfilerGraphPanel::OnDataGraphGraphOffsetChanged )
-				.OnViewModeChanged( this, &SProfilerGraphPanel::DataGraph_OnViewModeChanged )
+				SNew( SVerticalBox )
+
+				+ SVerticalBox::Slot()
+				.FillHeight( 1.0f )
+				[
+					SAssignNew( DataGraph, SDataGraph )
+					.OnGraphOffsetChanged( this, &SProfilerGraphPanel::OnDataGraphGraphOffsetChanged )
+					.OnViewModeChanged( this, &SProfilerGraphPanel::DataGraph_OnViewModeChanged )
+				]
+
+				+ SVerticalBox::Slot()
+				.AutoHeight()
+				[
+					SAssignNew( HorizontalScrollBar, SScrollBar )
+					.Orientation( Orient_Horizontal )
+					.AlwaysShowScrollbar( true )
+					.Visibility( EVisibility::Visible )
+					.Thickness( FVector2D( 8.0f, 8.0f ) )
+					.OnUserScrolled( this, &SProfilerGraphPanel::HorizontalScrollBar_OnUserScrolled )
+				]
 			]
 
-			+SVerticalBox::Slot()
-			.AutoHeight()
+			+SHorizontalBox::Slot()
+			.AutoWidth()
 			[
-				SAssignNew(ScrollBar,SScrollBar)
-				.Orientation( Orient_Horizontal	 )
+				SAssignNew( VerticalScrollBar, SScrollBar )
+				.Orientation( Orient_Vertical )
 				.AlwaysShowScrollbar( true )
 				.Visibility( EVisibility::Visible )
-				.Thickness( FVector2D(8.0f,8.0f) )
-				.OnUserScrolled( this, &SProfilerGraphPanel::ScrollBar_OnUserScrolled )
+				.Thickness( FVector2D( 8.0f, 8.0f ) )
+				.OnUserScrolled( this, &SProfilerGraphPanel::VerticalScrollBar_OnUserScrolled )
 			]
 		]	
 	];
 
-	ScrollBar->SetState( 0.0f, 1.0f );
+	HorizontalScrollBar->SetState( 0.0f, 1.0f );
+	VerticalScrollBar->SetState( 0.0f, 1.0f );
 
 	// Register ourselves with the profiler manager.
 	FProfilerManager::Get()->OnTrackedStatChanged().AddSP( this, &SProfilerGraphPanel::ProfilerManager_OnTrackedStatChanged );
@@ -65,12 +84,12 @@ void SProfilerGraphPanel::Tick( const FGeometry& AllottedGeometry, const double
 	UpdateInternals();
 }
 
-void SProfilerGraphPanel::ScrollBar_OnUserScrolled( float ScrollOffset )
+void SProfilerGraphPanel::HorizontalScrollBar_OnUserScrolled( float ScrollOffset )
 {
 	const float ThumbSizeFraction = FMath::Min( (float)NumVisiblePoints / (float)NumDataPoints, 1.0f );
 	ScrollOffset = FMath::Min( ScrollOffset, 1.0f-ThumbSizeFraction );
  
-	ScrollBar->SetState( ScrollOffset, ThumbSizeFraction );
+	HorizontalScrollBar->SetState( ScrollOffset, ThumbSizeFraction );
  
 	GraphOffset = FMath::Trunc( ScrollOffset * NumDataPoints );
 	DataGraph->ScrollTo( GraphOffset );
@@ -78,6 +97,13 @@ void SProfilerGraphPanel::ScrollBar_OnUserScrolled( float ScrollOffset )
 	FProfilerManager::Get()->GetProfilerWindow()->ProfilerMiniView->OnSelectionBoxMoved( GraphOffset, GraphOffset + NumVisiblePoints );
 }
 
+
+void SProfilerGraphPanel::VerticalScrollBar_OnUserScrolled( float ScrollOffset )
+{
+
+}
+
+
 void SProfilerGraphPanel::ProfilerManager_OnTrackedStatChanged( const FTrackedStat& TrackedStat, bool bIsTracked )
 {
 	if( bIsTracked )
@@ -125,7 +151,7 @@ void SProfilerGraphPanel::UpdateInternals()
 	if( FProfilerManager::Get()->IsLivePreview() )
 	{
 		// Scroll to the end.
-		ScrollBar_OnUserScrolled(1.0f);
+		HorizontalScrollBar_OnUserScrolled(1.0f);
 	}
 }
 
@@ -135,7 +161,7 @@ void SProfilerGraphPanel::SetScrollBarState()
 	// If the user can view 1/3 of the items in a single page, the maximum offset will be ~0.667f
 	const float ThumbSizeFraction = FMath::Min( (float)NumVisiblePoints / (float)NumDataPoints, 1.0f );
 	const float OffsetFraction = (float)GraphOffset / (float)NumDataPoints;
-	ScrollBar->SetState( OffsetFraction, ThumbSizeFraction );
+	HorizontalScrollBar->SetState( OffsetFraction, ThumbSizeFraction );
 }
 
 #undef LOCTEXT_NAMESPACE
diff --git a/Engine/Source/Developer/Profiler/Private/Widgets/SProfilerGraphPanel.h b/Engine/Source/Developer/Profiler/Private/Widgets/SProfilerGraphPanel.h
--- a/Engine/Source/Developer/Profiler/Private/Widgets/SProfilerGraphPanel.h
+++ b/Engine/Source/Developer/Profiler/Private/Widgets/SProfilerGraphPanel.h
@@ -43,7 +43,9 @@ protected:
 	 *
 	 * @param ScrollOffset - scroll offset as a fraction between 0 and 1.
 	 */
-	void ScrollBar_OnUserScrolled( float ScrollOffset );
+	void HorizontalScrollBar_OnUserScrolled( float ScrollOffset );
+
+	void VerticalScrollBar_OnUserScrolled( float ScrollOffset );
 
 	/** Called when the frame offset has been changed in the data graph widget. */
 	void OnDataGraphGraphOffsetChanged( int32 InFrameOffset );
@@ -65,8 +67,11 @@ protected:
 	/** Holds the data graph widget. */
 	TSharedPtr<SDataGraph> DataGraph;
 
-	/** Widget used for scrolling graphs. */
-	TSharedPtr<SScrollBar> ScrollBar;
+	/** Horizontal scroll bar, used for scrolling graphs. */
+	TSharedPtr<SScrollBar> HorizontalScrollBar;
+
+	/** Vertical scroll bar, used for scrolling graphs. */
+	TSharedPtr<SScrollBar> VerticalScrollBar;
 
 	/** Number of graph points. */
 	int32 NumDataPoints;
