#UE4.Profiler (misc tweaks, basic support for the new view mode, added profiler stream class)
- Misc tweaks to make profiler more base on the new stats system
- Added a new view mode to the data graph, for the regular stats file it will be LineIndexBased, for the raw stats file it will be ThreadViewTimeBased
- Added profiler stream, a new class that will hold all data read from the raw stats file, hopefully will replace the current profiler data provider/ data source later
- Added basic support for the second profiler view mode

diff --git a/Engine/Source/Developer/Profiler/Private/ProfilerManager.cpp b/Engine/Source/Developer/Profiler/Private/ProfilerManager.cpp
--- a/Engine/Source/Developer/Profiler/Private/ProfilerManager.cpp
+++ b/Engine/Source/Developer/Profiler/Private/ProfilerManager.cpp
@@ -26,7 +26,8 @@ DEFINE_STAT(STAT_PM_MemoryUsage);
 	Profiler manager implementation
 -----------------------------------------------------------------------------*/
 
-TSharedPtr<FProfilerManager> FProfilerManager::Instance = NULL;
+TSharedPtr<FProfilerManager> FProfilerManager::Instance = nullptr;
+FThreadSafeCounter FProfilerManager::ProcessingLock;
 
 struct FEventGraphSampleLess
 {
@@ -204,12 +205,12 @@ void FProfilerManager::LoadProfilerCapture( const FString& ProfilerCaptureFilepa
 	}
 
 	SessionInstancesUpdatedEvent.Broadcast();
-	ProfilerType = EProfilerSessionTypes::Offline;
-
+	ProfilerType = EProfilerSessionTypes::StatsFile;
+	
 	GetProfilerWindow()->ManageEventGraphTab( ProfilerInstanceID, true, ProfilerSession->GetName() );
+	SetViewMode( EProfilerViewMode::LineIndexBased );
 }
 
-
 void FProfilerManager::LoadRawStatsFile( const FString& RawStatsFileFileath )
 {
 	if( ActiveSession.IsValid() )
@@ -222,21 +223,22 @@ void FProfilerManager::LoadRawStatsFile( const FString& RawStatsFileFileath )
 	const FGuid ProfilerInstanceID = ProfilerSession->GetInstanceID();
 	ProfilerSessionInstances.Add( ProfilerInstanceID, ProfilerSession );
 
-	ProfilerSession->PrepareLoading();
-	RequestFilterAndPresetsUpdateEvent.Broadcast();
-	ProfilerSession_OnCaptureFileProcessed( ProfilerInstanceID );
-	TrackDefaultStats();
-
 	ProfilerSession->
-		SetOnCaptureFileProcessed( FProfilerSession::FCaptureFileProcessedDelegate::CreateSP( this, &FProfilerManager::ProfilerSession_OnCaptureFileProcessed ) )
-		.SetOnAddThreadTime( FProfilerSession::FAddThreadTimeDelegate::CreateSP( this, &FProfilerManager::ProfilerSession_OnAddThreadTime ) );
+		//SetOnCaptureFileProcessed( FProfilerSession::FCaptureFileProcessedDelegate::CreateSP( this, /&FProfilerManager::ProfilerSession_OnCaptureFileProcessed ) )
+		SetOnAddThreadTime( FProfilerSession::FAddThreadTimeDelegate::CreateSP( this, &FProfilerManager::ProfilerSession_OnAddThreadTime ) );
 
 	ProfilerSessionInstances.Add( ProfilerInstanceID, ProfilerSession );
 
-	SessionInstancesUpdatedEvent.Broadcast();
-	ProfilerType = EProfilerSessionTypes::OfflineRaw;
+	ProfilerSession->PrepareLoading();
+	RequestFilterAndPresetsUpdateEvent.Broadcast();
+	//ProfilerSession_OnCaptureFileProcessed( ProfilerInstanceID );
+	//TrackDefaultStats();
 
+	SessionInstancesUpdatedEvent.Broadcast();
+	ProfilerType = EProfilerSessionTypes::StatsFileRaw;
+	
 	GetProfilerWindow()->ManageEventGraphTab( ProfilerInstanceID, true, ProfilerSession->GetName() );
+	SetViewMode( EProfilerViewMode::ThreadViewTimeBased );
 }
 
 
@@ -549,6 +551,7 @@ void FProfilerManager::ClearStatsAndInstances()
 	CloseAllEventGraphTabs();
 
 	ProfilerType = EProfilerSessionTypes::InvalidOrMax;
+	ViewMode = EProfilerViewMode::InvalidOrMax;
 	SetDataPreview( false );
 	bLivePreview = false;
 	SetDataCapture( false );
@@ -743,4 +746,14 @@ void FProfilerManager::ProfilerSession_OnAddThreadTime( int32 FrameIndex, const
 	}
 }
 
+void FProfilerManager::SetViewMode( EProfilerViewMode::Type NewViewMode )
+{
+	if( NewViewMode != ViewMode )
+	{
+		// Broadcast.
+		OnViewModeChangedEvent.Broadcast( NewViewMode );
+		ViewMode = NewViewMode;
+	}
+}
+
 #undef LOCTEXT_NAMESPACE
\ No newline at end of file
diff --git a/Engine/Source/Developer/Profiler/Private/ProfilerManager.h b/Engine/Source/Developer/Profiler/Private/ProfilerManager.h
--- a/Engine/Source/Developer/Profiler/Private/ProfilerManager.h
+++ b/Engine/Source/Developer/Profiler/Private/ProfilerManager.h
@@ -85,6 +85,7 @@ public:
 	bool bIsDefault;
 };
 
+
 /** Contains basic information about tracked stat. */
 class FTrackedStat //: public FNoncopyable, public TSharedFromThis<FTrackedStat>
 {
@@ -138,6 +139,19 @@ public:
 	FProfilerManager
 -----------------------------------------------------------------------------*/
 
+namespace EProfilerViewMode
+{
+	enum Type
+	{
+		/** Regular line graphs, for the regular stats file. */
+		LineIndexBased,
+		/** Thread view graph, for the raw stats file. */
+		ThreadViewTimeBased,
+		/** Invalid enum type, may be used as a number of enumerations. */
+		InvalidOrMax,
+	};
+}
+
 /** 
  ** This class manages following areas:
  **		Connecting/disconnecting to source or device through session manager
@@ -153,6 +167,12 @@ class FProfilerManager
 {
 	friend class FProfilerActionManager;
 
+	/**
+	 *	Global processing lock, protects from destroying the profiler while task graph is still working on the stuff.
+	 *	Assumes that there is only one instance of the profiler.
+	 */
+	static FThreadSafeCounter ProcessingLock;
+
 public:
 	/**
 	 * Creates a profiler manager, only one distance can exists.
@@ -193,6 +213,18 @@ public:
 		FProfilerManager::Instance.Reset();
 	}
 
+	/** Increases processing lock count. */
+	static void IncrementProcessingLock()
+	{
+		ProcessingLock.Increment();
+	}
+
+	/** Decreases processing lock count, once reaches zero the profiler can proceed. */
+	static void DecrementProcessingLock()
+	{
+		ProcessingLock.Decrement();
+	}
+
 protected:
 	/**
 	 * Finishes initialization of the profiler manager
@@ -404,6 +436,17 @@ public:
 		Events declarations
 	-----------------------------------------------------------------------------*/
 	
+public:
+	DECLARE_EVENT_OneParam( FProfilerManager, FViewModeChangedEvent, EProfilerViewMode::Type /*NewViewMode*/ );
+	FViewModeChangedEvent& OnViewModeChanged()
+	{
+		return OnViewModeChangedEvent;
+	}
+	
+protected:
+	/** The event to execute when the profiler loaded a new stats files and the view mode needs to be changed. */
+	FViewModeChangedEvent OnViewModeChangedEvent;
+	
 
 public:
 	/**
@@ -538,6 +581,9 @@ public:
 		return ProfilerWindow.Pin();
 	}
 
+	/** Sets a new view mode for the profiler. */
+	void SetViewMode( EProfilerViewMode::Type NewViewMode );
+
 protected:
 	/** The delegate to be invoked when this profiler manager ticks. */
 	FTickerDelegate OnTick;
@@ -589,6 +635,9 @@ protected:
 	/** Profiler session type that is currently initialized. */
 	EProfilerSessionTypes::Type ProfilerType;
 
+	/** Profiler view mode. */
+	EProfilerViewMode::Type ViewMode;
+
 	// TODO: Bool should be replaces with type similar to ESlateCheckBoxState {Checked,Unchecked,Undertermined}
 
 	/** True, if the profiler is currently showing the latest data, only valid if profiler is connected to network based session. */
diff --git a/Engine/Source/Developer/Profiler/Private/ProfilerModule.cpp b/Engine/Source/Developer/Profiler/Private/ProfilerModule.cpp
--- a/Engine/Source/Developer/Profiler/Private/ProfilerModule.cpp
+++ b/Engine/Source/Developer/Profiler/Private/ProfilerModule.cpp
@@ -16,7 +16,7 @@ public:
 	virtual TSharedRef<SWidget> CreateProfilerWindow( const ISessionManagerRef& InSessionManager, const TSharedRef<SDockTab>& ConstructUnderMajorTab ) OVERRIDE
 	{
 		ProfilerManager = FProfilerManager::Initialize( InSessionManager );
-		TSharedRef<SProfilerWindow> ProfilerWindow = SNew(SProfilerWindow, InSessionManager);
+		TSharedRef<SProfilerWindow> ProfilerWindow = SNew(SProfilerWindow);
 		FProfilerManager::Get()->AssignProfilerWindow( ProfilerWindow );
 		// Register OnTabClosed to handle profiler manager shutdown.
 		ConstructUnderMajorTab->SetOnTabClosed( SDockTab::FOnTabClosedCallback::CreateRaw(this, &FProfilerModule::Shutdown) );
diff --git a/Engine/Source/Developer/Profiler/Private/ProfilerPrivatePCH.h b/Engine/Source/Developer/Profiler/Private/ProfilerPrivatePCH.h
--- a/Engine/Source/Developer/Profiler/Private/ProfilerPrivatePCH.h
+++ b/Engine/Source/Developer/Profiler/Private/ProfilerPrivatePCH.h
@@ -29,6 +29,8 @@
 // TODO: Remove dependencies between generic profiler code and profiler ui code
 
 // Common headers and managers.
+#include "ProfilerStream.h"
+
 #include "ProfilerSample.h"
 #include "ProfilerFPSAnalyzer.h"
 #include "ProfilerDataProvider.h"
diff --git a/Engine/Source/Developer/Profiler/Private/ProfilerSession.cpp b/Engine/Source/Developer/Profiler/Private/ProfilerSession.cpp
--- a/Engine/Source/Developer/Profiler/Private/ProfilerSession.cpp
+++ b/Engine/Source/Developer/Profiler/Private/ProfilerSession.cpp
@@ -83,7 +83,7 @@ FProfilerSession::FProfilerSession( const FString& InDataFilepath )
 , EventGraphDataCurrent( MakeShareable( new FEventGraphData() ) )
 , CreationTime( FDateTime::Now() )
 , DataFilepath( InDataFilepath.Replace( *FStatConstants::StatsFileExtension, TEXT( "" ) ) )
-, SessionType( EProfilerSessionTypes::Offline )
+, SessionType( EProfilerSessionTypes::StatsFile )
 , SessionInstanceInfo( nullptr )
 , SessionInstanceID( FGuid::NewGuid() )
 , bDataPreviewing( false )
@@ -239,7 +239,7 @@ bool FProfilerSession::HandleTicker( float DeltaTime )
 		bRequestStatMetadataUpdate = false;
 	}
 
-	const int32 FramesNumToProcess = SessionType == EProfilerSessionTypes::Offline ? 240 : 2;
+	const int32 FramesNumToProcess = SessionType == EProfilerSessionTypes::StatsFile ? 240 : 2;
 
 	for( int32 FrameNumber = 0; FrameNumber < FMath::Min(FrameToProcess.Num(),FramesNumToProcess); FrameNumber++ )
 	{	
@@ -359,7 +359,7 @@ bool FProfilerSession::HandleTicker( float DeltaTime )
 		FrameToProfilerDataMapping.Remove( FrameIndex );
 	}	
 
-	if( SessionType == EProfilerSessionTypes::Offline )
+	if( SessionType == EProfilerSessionTypes::StatsFile )
 	{
 		if( FrameToProcess.Num() == 0 && bHasAllProfilerData )
 		{
@@ -437,7 +437,7 @@ void FProfilerSession::PopulateHierarchy_Recurrent
 -----------------------------------------------------------------------------*/
 
 FRawProfilerSession::FRawProfilerSession( const FString& InRawStatsFileFileath )
-: FProfilerSession( EProfilerSessionTypes::OfflineRaw, nullptr, FGuid::NewGuid(), InRawStatsFileFileath.Replace( *FStatConstants::StatsFileRawExtension, TEXT( "" ) ) )
+: FProfilerSession( EProfilerSessionTypes::StatsFileRaw, nullptr, FGuid::NewGuid(), InRawStatsFileFileath.Replace( *FStatConstants::StatsFileRawExtension, TEXT( "" ) ) )
 , CurrentMiniViewFrame( 0 )
 {
 	OnTick = FTickerDelegate::CreateRaw( this, &FRawProfilerSession::HandleTicker );
@@ -538,9 +538,9 @@ void FRawProfilerSession::PrepareLoading()
 		}
 		int64 FrameOffset0 = Stream.FramesInfo[0].FrameOffset;
 		FileReader->Seek( FrameOffset0 );
-#if	0
+#if	1
 		// Read the raw stats messages.
-		for( int32 FrameIndex = 0; FrameIndex < 30/*Stream.FramesInfo.Num()*/; ++FrameIndex )
+		for( int32 FrameIndex = 0; FrameIndex < Stream.FramesInfo.Num(); ++FrameIndex )
 		{
 			FStatPacketArray Packet;
 
@@ -557,17 +557,33 @@ void FRawProfilerSession::PrepareLoading()
 				Packet.Packets.Add( ToRead );
 			}
 
-			ProcessStatPacketArray( Packet );
+			FProfilerFrame* ProfilerFrame = ProcessStatPacketArray( Packet );
 
 			TMap<uint32, int64> ThreadCycles;
 			// Serialize thread cycles.
 			*FileReader << ThreadCycles;
+			
+			// Update mini-view.
+			// Convert from cycles to ms.
+			TMap<uint32, float> ThreadTimesMS;
+			for( auto InnerIt = ThreadCycles.CreateConstIterator(); InnerIt; ++InnerIt )
+			{
+				ThreadTimesMS.Add( InnerIt.Key(), StatMetaData->ConvertCyclesToMS( InnerIt.Value() ) );
+			}
+
+			// Pass the reference to the stats' metadata.
+			OnAddThreadTime.ExecuteIfBound( FrameIndex, ThreadTimesMS, StatMetaData );
+
+			ProfilerFrame->ThreadTimesMS = ThreadTimesMS;
+			ProfilerStream.AddProfilerFrame( ProfilerFrame );
 		}
 #endif // 0
 	}
 
+	const int64 AllocatedSize = ProfilerStream.GetAllocatedSize();
+
 	// We have the whole metadata and basic information about the raw stats file, start ticking the profiler session.
-	FTicker::GetCoreTicker().AddTicker( OnTick, 0.25f );
+	//FTicker::GetCoreTicker().AddTicker( OnTick, 0.25f );
 
 #if	0
 	if( SessionType == EProfilerSessionTypes::OfflineRaw )
@@ -692,46 +708,23 @@ void FProfilerStatMetaData::UpdateFromStatsState( const FStatsThreadState& Stats
 				RenderThreadIDs.Add( ThreadID );
 			}
 		}
-	}
-
-	// @TODO yrx 2014-03-24 Fix this.
-	SecondsPerCycle = FPlatformTime::GetSecondsPerCycle();
-}
 
-/** A simple stats stack node, used as a proxy class to convert data from the stats into the profiler. */
-struct FStatStackNode
-{
-	/** Short name. */
-	FName NodeName;
-	int64 CyclesStart;
-	int64 CyclesEnd;
-	TArray<FStatStackNode*> Children;
-
-	/** Index of this node in the data provider's collection. */
-	uint32 SampleIndex;
-
-	explicit FStatStackNode( const FStatMessage& StatMessage, uint32 InSampleIndex )
-		: NodeName( StatMessage.NameAndInfo.GetShortName() )
-		, CyclesStart( StatMessage.GetValue_int64() )
-		, CyclesEnd( 0 )
-		, SampleIndex( InSampleIndex )
-	{}
-
-	~FStatStackNode()
-	{
-		for( const auto& Child : Children )
+		if( StatName == TEXT( "STAT_SecondsPerCycle" ) )
 		{
-			delete Child;
+			SecondsPerCycle = LongName.GetValue_double();
 		}
 	}
-};
 
-void FRawProfilerSession::ProcessStatPacketArray( const FStatPacketArray& StatPacketArray )
+	SecondsPerCycle = FPlatformTime::GetSecondsPerCycle();
+}
+
+FProfilerFrame* FRawProfilerSession::ProcessStatPacketArray( const FStatPacketArray& StatPacketArray )
 {
 	// @TODO yrx 2014-03-24 Standardize thread names and id
+	// @TODO yrx 2014-04-22 Remove all references to the data provider, event graph etc once data graph can visualize.
 
 	// Raw stats callstack for this stat packet array.
-	TMap<FName,FStatStackNode> ThreadNodes;
+	TMap<FName,FProfilerStackNode*> ThreadNodes;
 
 	const FProfilerStatMetaDataRef MetaData = GetMetaData();
 	
@@ -748,15 +741,18 @@ void FRawProfilerSession::ProcessStatPacketArray( const FStatPacketArray& StatPa
 		const FName ThreadFName = StatsThreadStats.Threads.FindChecked( StatPacket.ThreadId );
 		const uint32 NewThreadID = MetaData->ThreadIDtoStatID.FindChecked( StatPacket.ThreadId );
 
+		// @TODO yrx 2014-04-18 At this moment only for the game thread.
 		if( StatPacket.ThreadType != EThreadType::Game )
 		{
 			continue;
 		}
 
-		FStatStackNode* ThreadNode = ThreadNodes.Find( ThreadFName );
+		FProfilerStackNode* ThreadNode = ThreadNodes.FindRef( ThreadFName );
 		if( !ThreadNode )
 		{
-			FStatMessage ThreadMessage( ThreadFName, EStatDataType::ST_int64, nullptr, TEXT( "" ), true, true );
+			FString ThreadIdName = FStatsUtils::BuildUniqueThreadName( StatPacket.ThreadId );
+			FStatMessage ThreadMessage( ThreadFName, EStatDataType::ST_int64, STAT_GROUP_TO_FStatGroup( STATGROUP_Threads )::GetGroupName(), *ThreadIdName, true, true );
+			//FStatMessage ThreadMessage( ThreadFName, EStatDataType::ST_int64, nullptr, TEXT( "" ), true, true );
 			ThreadMessage.NameAndInfo.SetFlag( EStatMetaFlags::IsPackedCCAndDuration, true );
 			ThreadMessage.Clear();
 
@@ -772,7 +768,7 @@ void FRawProfilerSession::ProcessStatPacketArray( const FStatPacketArray& StatPa
 				FrameRootSampleIndex
 			);
 
-			ThreadNode = &ThreadNodes.Add( ThreadFName, FStatStackNode( ThreadMessage, ThreadRootSampleIndex ) );
+			ThreadNode = ThreadNodes.Add( ThreadFName, new FProfilerStackNode( nullptr, ThreadMessage, ThreadRootSampleIndex ) );
 		}
 
 		if( StatPacket.ThreadType == EThreadType::Game )
@@ -781,9 +777,9 @@ void FRawProfilerSession::ProcessStatPacketArray( const FStatPacketArray& StatPa
 		}
 
 		TArray<const FStatMessage*> StartStack;
-		TArray<FStatStackNode*> Stack;
+		TArray<FProfilerStackNode*> Stack;
 		Stack.Add( ThreadNode );
-		FStatStackNode* Current = Stack.Last();
+		FProfilerStackNode* Current = Stack.Last();
 
 		const TArray<FStatMessage>& Data = StatPacket.StatMessages;
 		for( int32 Index = 0; Index < Data.Num(); Index++ )
@@ -800,7 +796,7 @@ void FRawProfilerSession::ProcessStatPacketArray( const FStatPacketArray& StatPa
 				check( Item.NameAndInfo.GetFlag( EStatMetaFlags::IsCycle ) );
 				if( Op == EStatOperation::CycleScopeStart )
 				{
-					FStatStackNode* ChildNode = new FStatStackNode( Item, -1 );
+					FProfilerStackNode* ChildNode = new FProfilerStackNode( Current, Item, -1 );
 					Current->Children.Add( ChildNode );
 
 					// Add a child sample.
@@ -826,7 +822,12 @@ void FRawProfilerSession::ProcessStatPacketArray( const FStatPacketArray& StatPa
 					const FStatMessage ScopeEnd = Item;
 					const int64 Delta = int32( uint32( ScopeEnd.GetValue_int64() ) - uint32( ScopeStart.GetValue_int64() ) );
 					Current->CyclesEnd = Current->CyclesStart + int64( Delta );
-					FStatStackNode* ChildNode = Current;
+
+					Current->TimeMSStart = MetaData->ConvertCyclesToMS( Current->CyclesStart );
+					Current->TimeMSEnd = MetaData->ConvertCyclesToMS( Current->CyclesEnd );
+					check( Current->TimeMSEnd >= Current->TimeMSStart );
+
+					FProfilerStackNode* ChildNode = Current;
 
 					// Update the child sample's DurationMS.
 					MutableCollection[ChildNode->SampleIndex].SetDurationMS( MetaData->ConvertCyclesToMS( Delta ) );
@@ -841,13 +842,15 @@ void FRawProfilerSession::ProcessStatPacketArray( const FStatPacketArray& StatPa
 	// Calculate thread times.
 	for( auto It = ThreadNodes.CreateIterator(); It; ++It )
 	{
-		FStatStackNode& ThreadNode = It.Value();
+		FProfilerStackNode& ThreadNode = *It.Value();
 		const int32 ChildrenNum = ThreadNode.Children.Num();
 		if( ChildrenNum > 0 )
 		{
 			const int32 LastChildIndex = ThreadNode.Children.Num() - 1;
 			ThreadNode.CyclesStart = ThreadNode.Children[0]->CyclesStart;
 			ThreadNode.CyclesEnd = ThreadNode.Children[LastChildIndex]->CyclesEnd;
+			ThreadNode.TimeMSStart = MetaData->ConvertCyclesToMS( ThreadNode.CyclesStart );
+			ThreadNode.TimeMSEnd = MetaData->ConvertCyclesToMS( ThreadNode.CyclesEnd );
 
 			FProfilerSample& ProfilerSample = MutableCollection[ThreadNode.SampleIndex];
 			ProfilerSample.SetStartAndEndMS( MetaData->ConvertCyclesToMS( ThreadNode.CyclesStart ), MetaData->ConvertCyclesToMS( ThreadNode.CyclesEnd ) );
@@ -856,7 +859,7 @@ void FRawProfilerSession::ProcessStatPacketArray( const FStatPacketArray& StatPa
 
 	// Get the game thread time.
 	check( GameThreadFName != NAME_None );
-	const FStatStackNode& GameThreadNode = ThreadNodes.FindChecked( GameThreadFName );
+	const FProfilerStackNode& GameThreadNode = *ThreadNodes.FindChecked( GameThreadFName );
 	const double GameThreadStartMS = MetaData->ConvertCyclesToMS( GameThreadNode.CyclesStart );
 	const double GameThreadEndMS = MetaData->ConvertCyclesToMS( GameThreadNode.CyclesEnd );
 	MutableCollection[FrameRootSampleIndex].SetStartAndEndMS( GameThreadStartMS, GameThreadEndMS );
@@ -870,4 +873,13 @@ void FRawProfilerSession::ProcessStatPacketArray( const FStatPacketArray& StatPa
  
 	// Update aggregated events.
 	UpdateAggregatedEventGraphData( LastFrameIndex );
+
+	FProfilerFrame* ProfilerFrame = new FProfilerFrame();
+	for( auto It = ThreadNodes.CreateIterator(); It; ++It )
+	{
+		ProfilerFrame->Root->Children.Add( It.Value() );
+		// @TODO yrx 2014-04-22 Sort children GameThread, RenderThread, Other children sorted by total time.
+	}
+
+	return ProfilerFrame;
 }
\ No newline at end of file
diff --git a/Engine/Source/Developer/Profiler/Private/ProfilerSession.h b/Engine/Source/Developer/Profiler/Private/ProfilerSession.h
--- a/Engine/Source/Developer/Profiler/Private/ProfilerSession.h
+++ b/Engine/Source/Developer/Profiler/Private/ProfilerSession.h
@@ -13,12 +13,13 @@ namespace EProfilerSessionTypes
 	{
 		/** Based on the live connection. */
 		Live,
+		LiveRaw,
 
 		/** Based on the regular stats file. */
-		Offline,
+		StatsFile,
 
 		/** Based on the raw stats file. */
-		OfflineRaw,
+		StatsFileRaw,
 
 		Combined,
 		Summary,
@@ -40,7 +41,7 @@ namespace EProfilerSessionTypes
 		switch( ProfilerSessionType )
 		{
 			case Live: return FString("Live");
-			case Offline: return FString("Offline");
+			case StatsFile: return FString("Offline");
 
 			default: return FString("InvalidOrMax");
 		}
@@ -1021,11 +1022,11 @@ public:
 		{
 			SessionName = FString::Printf( TEXT("%s"), *SessionInstanceInfo->GetInstanceName() );
 		}
-		else if( SessionType == EProfilerSessionTypes::Offline )
+		else if( SessionType == EProfilerSessionTypes::StatsFile )
 		{
 			SessionName = FString::Printf( TEXT("%s"), *DataFilepath );
 		}
-		else if( SessionType == EProfilerSessionTypes::OfflineRaw )
+		else if( SessionType == EProfilerSessionTypes::StatsFileRaw )
 		{
 			SessionName = FString::Printf( TEXT( "%s" ), *DataFilepath );
 		}
@@ -1271,6 +1272,9 @@ class FRawProfilerSession : public FProfilerSession
 	friend class FProfilerManager;
 	friend class FProfilerActionManager;
 
+	/** Profiler stream that contains all read raw profiler frames. */
+	FProfilerStream ProfilerStream;
+
 	/** Stats thread state, mostly used to manage the stats metadata. */
 	FStatsThreadState StatsThreadStats;
 	FStatsReadStream Stream;
@@ -1297,5 +1301,5 @@ public:
 	 *	Process all stats packets and convert them to data accessible by the profiler.
 	 *	Temporary version, will be optimized later.
 	 */
-	void ProcessStatPacketArray( const FStatPacketArray& PacketArray );
+	FProfilerFrame* ProcessStatPacketArray( const FStatPacketArray& PacketArray );
 };
\ No newline at end of file
diff --git a/Engine/Source/Developer/Profiler/Private/ProfilerStream.cpp b/Engine/Source/Developer/Profiler/Private/ProfilerStream.cpp
new file mode 100644
--- /dev/null
+++ b/Engine/Source/Developer/Profiler/Private/ProfilerStream.cpp
@@ -0,0 +1,3 @@
+// Copyright 1998-2014 Epic Games, Inc. All Rights Reserved.
+
+#include "ProfilerPrivatePCH.h"
\ No newline at end of file
diff --git a/Engine/Source/Developer/Profiler/Private/ProfilerStream.h b/Engine/Source/Developer/Profiler/Private/ProfilerStream.h
new file mode 100644
--- /dev/null
+++ b/Engine/Source/Developer/Profiler/Private/ProfilerStream.h
@@ -0,0 +1,210 @@
+// Copyright 1998-2014 Epic Games, Inc. All Rights Reserved.
+
+#pragma once
+
+/*-----------------------------------------------------------------------------
+	Basic structures
+-----------------------------------------------------------------------------*/
+
+/** Profiler stack node, used to store the whole call stack for one frame. */
+struct FProfilerStackNode
+{
+	FProfilerStackNode* Parent;
+	TArray<FProfilerStackNode*> Children;
+
+	/** Short name. */
+	FName StatName;
+	FName LongName;
+
+	int64 CyclesStart;
+	int64 CyclesEnd;
+
+	double TimeMSStart;
+	double TimeMSEnd;
+
+	/** Index of this node in the data provider's collection. */
+	uint32 SampleIndex;
+
+	/** Initializes thread root node. */
+	FProfilerStackNode()
+		: Parent( nullptr )
+		, StatName( FStatConstants::NAME_ThreadRoot )
+		, LongName( FStatConstants::NAME_ThreadRoot )
+		, CyclesStart( 0 )
+		, CyclesEnd( 0 )
+		, TimeMSStart( 0.0f )
+		, TimeMSEnd( 0.0f )
+		, SampleIndex( 0 )
+	{}
+
+	/** Initializes children node. */
+	FProfilerStackNode( FProfilerStackNode* InParent, const FStatMessage& StatMessage, uint32 InSampleIndex )
+		: Parent( InParent )
+		, StatName( StatMessage.NameAndInfo.GetShortName() )
+		, LongName( StatMessage.NameAndInfo.GetRawName() )
+		, CyclesStart( StatMessage.GetValue_int64() )
+		, CyclesEnd( 0 )
+		, TimeMSStart( 0.0f )
+		, TimeMSEnd( 0.0f )
+		, SampleIndex( InSampleIndex )
+	{}
+
+	~FProfilerStackNode()
+	{
+		for( const auto& Child : Children )
+		{
+			delete Child;
+		}
+	}
+
+	/** Calculates allocated size by all stacks node, only valid for the thread root node. */
+	int64 GetAllocatedSize() const
+	{
+		int64 AllocatedSize = sizeof(*this) + Children.GetAllocatedSize();
+		for( const auto& StackNode : Children )
+		{
+			AllocatedSize += StackNode->GetAllocatedSize();
+		}
+		return AllocatedSize;
+	}
+};
+
+enum
+{
+	STACKNODE_INVALID = 0,
+	STACKNODE_VALID = 1,
+};
+
+/** Profiler frame. */
+struct FProfilerFrame
+{
+	/** Root node. */
+	FProfilerStackNode* Root;
+
+	/** Thread times in milliseconds for this frame. */
+	TMap<uint32, float> ThreadTimesMS;
+
+	// @TODO yrx 2014-04-22 Non-frame stats like accumulator, counters, memory etc?
+
+	/**
+	*	Indicates whether this node is in the memory, move to thread root node later.
+	*	Valid only for the root stack node.
+	*	This set by one thread and accessed by another thread, there is no thread contention.
+	*/
+	FThreadSafeCounter AccessLock;
+
+	FProfilerFrame()
+		: Root( new FProfilerStackNode() )
+		, AccessLock( STACKNODE_INVALID )
+	{}
+
+	~FProfilerFrame()
+	{
+		FreeMemory();
+	}
+
+	void MarkAsValid()
+	{
+		const int32 OldLock = AccessLock.Set( STACKNODE_VALID );
+		check( OldLock == STACKNODE_INVALID );
+	}
+
+	void MarkAsInvalid()
+	{
+		const int32 OldLock = AccessLock.Set( STACKNODE_INVALID );
+		check( OldLock == STACKNODE_VALID );
+	}
+
+	bool IsValid() const
+	{
+		return AccessLock.GetValue() == STACKNODE_VALID;
+	}
+
+	int64 GetAllocatedSize() const
+	{
+		return ThreadTimesMS.GetAllocatedSize() + Root ? Root->GetAllocatedSize() : 0;
+	}
+
+	void FreeMemory()
+	{
+		delete Root;
+		Root = nullptr;
+
+		ThreadTimesMS.Empty();
+	}
+};
+
+/** Contains all processed profiler's frames. */
+// @TODO yrx 2014-04-22 Array frames to real frames conversion.
+// @TODO yrx 2014-04-22 Can be done fully lock free and thread safe, but later.
+struct FProfilerStream
+{
+protected:
+
+	/** History frames collected so far or read from the file. */
+	TChunkedArray<FProfilerFrame*> Frames;
+
+	/** Critical section. */
+	mutable FCriticalSection CriticalSection;
+
+
+public:
+	void AddProfilerFrame( FProfilerFrame* ProfilerFrame )
+	{
+		FScopeLock Lock( &CriticalSection );
+		Frames.AddElement( ProfilerFrame );
+	}
+
+	/**
+	* @return a pointer to the profiler frame, once obtained can be used until end of the profiler session.
+	*/
+	FProfilerFrame* GetProfilerFrame( int32 FrameIndex ) const
+	{
+		FScopeLock Lock( &CriticalSection );
+		return Frames[FrameIndex];
+	}
+
+	int64 GetAllocatedSize() const
+	{
+		FScopeLock Lock( &CriticalSection );
+		int64 AllocatedSize = 0;
+
+		for( int32 FrameIndex = 0; FrameIndex < Frames.Num(); FrameIndex++ )
+		{
+			const FProfilerFrame* ProfilerFrame = Frames[FrameIndex];
+			if( ProfilerFrame->IsValid() )
+			{
+				AllocatedSize += ProfilerFrame->GetAllocatedSize();
+			}
+		}
+
+		return AllocatedSize;
+	}
+};
+
+/**
+*	Profiler UI stack node.
+*	Similar to the profiler stack node, but contains data prepared and optimized for the UI
+*/
+struct FProfilerUIStackNode
+{
+	/** Frames used to generate this data, at least 1 frame. */
+	TArray<int64> Frames;
+
+	FProfilerUIStackNode* Parent;
+	TArray<FProfilerUIStackNode*> Children;
+
+	/** Short name. */
+	FName StatName;
+	FName LongName;
+
+	double TimeMSStart;
+	double TimeMSEnd;
+};
+
+
+struct FProfilerUIStream
+{
+	TArray<FProfilerFrame*> Frames;
+};
+
diff --git a/Engine/Source/Developer/Profiler/Private/Widgets/SEventGraph.cpp b/Engine/Source/Developer/Profiler/Private/Widgets/SEventGraph.cpp
--- a/Engine/Source/Developer/Profiler/Private/Widgets/SEventGraph.cpp
+++ b/Engine/Source/Developer/Profiler/Private/Widgets/SEventGraph.cpp
@@ -733,6 +733,17 @@ SEventGraph::SEventGraph()
 	: CurrentStateIndex( 0 )
 {}
 
+
+SEventGraph::~SEventGraph()
+{
+	// Remove ourselves from the profiler manager.
+	if( FProfilerManager::Get().IsValid() )
+	{
+		FProfilerManager::Get()->OnViewModeChanged().RemoveAll( this );
+	}
+}
+
+
 /*-----------------------------------------------------------------------------
 	Event graph construction related functions
 -----------------------------------------------------------------------------*/
@@ -804,7 +815,7 @@ void SEventGraph::Construct( const FArguments& InArgs )
 		+SVerticalBox::Slot()
 		.AutoHeight()
 		[		
-			SNew(SBox)
+			SAssignNew(FunctionDetailsBox,SBox)
 			.HeightOverride( 224.0f )
 			[
 				SNew( SBorder )
@@ -911,6 +922,8 @@ void SEventGraph::Construct( const FArguments& InArgs )
 
 	InitializeAndShowHeaderColumns();
 	BindCommands();
+
+	FProfilerManager::Get()->OnViewModeChanged().AddSP( this, &SEventGraph::ProfilerManager_OnViewModeChanged );
 }
 
 TSharedRef<SWidget> SEventGraph::GetToggleButtonForEventGraphType( const EEventGraphTypes::Type EventGraphType, const FName BrushName )
@@ -3346,6 +3359,20 @@ bool SEventGraph::SelectAllFrames_CanExecute() const
 	return IsEventGraphStatesHistoryValid();
 }
 
+void SEventGraph::ProfilerManager_OnViewModeChanged( EProfilerViewMode::Type NewViewMode )
+{
+	if( NewViewMode == EProfilerViewMode::LineIndexBased )
+	{
+		FunctionDetailsBox->SetVisibility( EVisibility::Visible );
+		FunctionDetailsBox->SetEnabled( true );
+	}
+	else if( NewViewMode == EProfilerViewMode::ThreadViewTimeBased )
+	{
+		FunctionDetailsBox->SetVisibility( EVisibility::Collapsed );
+		FunctionDetailsBox->SetEnabled( false );
+	}
+}
+
 #undef LOCTEXT_NAMESPACE
 
 /*-----------------------------------------------------------------------------
diff --git a/Engine/Source/Developer/Profiler/Private/Widgets/SEventGraph.h b/Engine/Source/Developer/Profiler/Private/Widgets/SEventGraph.h
--- a/Engine/Source/Developer/Profiler/Private/Widgets/SEventGraph.h
+++ b/Engine/Source/Developer/Profiler/Private/Widgets/SEventGraph.h
@@ -164,6 +164,9 @@ public:
 	/** Default constructor. */
 	SEventGraph();
 
+	/** Destructor. */
+	~SEventGraph();
+
 	/**
 	 * Construct this widget.
 	 *
@@ -207,6 +210,9 @@ public:
 	{
 		return EventGraphRestoredFromHistoryEvent;
 	}
+
+protected:
+	void ProfilerManager_OnViewModeChanged( EProfilerViewMode::Type NewViewMode );
 	
 protected:
 	/** The event to execute when the event graph has been restored from the history. */
@@ -545,6 +551,8 @@ protected:
 	/** External scrollbar used to synchronize tree view position. */
 	TSharedPtr<SScrollBar> ExternalScrollbar;
 
+	TSharedPtr<SBox> FunctionDetailsBox;
+
 	/** Holds the tree view header row widget which display all columns in the tree view. */
 	TSharedPtr< SHeaderRow > TreeViewHeaderRow;
 
diff --git a/Engine/Source/Developer/Profiler/Private/Widgets/SFiltersAndPresets.cpp b/Engine/Source/Developer/Profiler/Private/Widgets/SFiltersAndPresets.cpp
--- a/Engine/Source/Developer/Profiler/Private/Widgets/SFiltersAndPresets.cpp
+++ b/Engine/Source/Developer/Profiler/Private/Widgets/SFiltersAndPresets.cpp
@@ -101,7 +101,7 @@ protected:
 
 	void AddHeader( const TSharedRef<SGridPanel>& Grid, int32& RowPos )
 	{
-		const FString InstanceName = ProfilerSession->GetSessionType() == EProfilerSessionTypes::Offline ? FPaths::GetBaseFilename( ProfilerSession->GetName() ) : ProfilerSession->GetName();
+		const FString InstanceName = ProfilerSession->GetSessionType() == EProfilerSessionTypes::StatsFile ? FPaths::GetBaseFilename( ProfilerSession->GetName() ) : ProfilerSession->GetName();
 
 		Grid->AddSlot( 0, RowPos++ )
 		.Padding( 2.0f )
@@ -581,7 +581,8 @@ SFiltersAndPresets::~SFiltersAndPresets()
 	// Remove ourselves from the profiler manager.
 	if( FProfilerManager::Get().IsValid() )
 	{
-		FProfilerManager::Get()->OnRequestFilterAndPresetsUpdate().RemoveAll(this);
+		FProfilerManager::Get()->OnRequestFilterAndPresetsUpdate().RemoveAll( this );
+		FProfilerManager::Get()->OnViewModeChanged().RemoveAll( this );
 	}
 }
 
@@ -742,6 +743,8 @@ void SFiltersAndPresets::Construct( const FArguments& InArgs )
 	// Register ourselves with the profiler manager.
 	FProfilerManager::Get()->OnRequestFilterAndPresetsUpdate().AddSP( this, &SFiltersAndPresets::ProfilerManager_OnRequestFilterAndPresetsUpdate );
 
+	FProfilerManager::Get()->OnViewModeChanged().AddSP( this, &SFiltersAndPresets::ProfilerManager_OnViewModeChanged );
+
 	// Create the search filters: text based, stat type based etc.
 	// @TODO: HandleItemToStringArray should be moved somewhere else
 	GroupAndStatTextFilter = MakeShareable( new FGroupAndStatTextFilter( FGroupAndStatTextFilter::FItemToStringArray::CreateSP( this, &SFiltersAndPresets::HandleItemToStringArray ) ) );
@@ -757,6 +760,22 @@ void SFiltersAndPresets::ProfilerManager_OnRequestFilterAndPresetsUpdate()
 	bUpdateOnNextTick = true;
 }
 
+
+void SFiltersAndPresets::ProfilerManager_OnViewModeChanged( EProfilerViewMode::Type NewViewMode )
+{
+	if( NewViewMode == EProfilerViewMode::LineIndexBased )
+	{
+		SetVisibility( EVisibility::Visible );
+		SetEnabled( true );
+	}
+	else if( NewViewMode == EProfilerViewMode::ThreadViewTimeBased )
+	{
+		SetVisibility( EVisibility::Collapsed );
+		SetEnabled( false );
+	}
+}
+
+
 void SFiltersAndPresets::Tick( const FGeometry& AllottedGeometry, const double InCurrentTime, const float InDeltaTime )
 {
 	if( bUpdateOnNextTick )
diff --git a/Engine/Source/Developer/Profiler/Private/Widgets/SFiltersAndPresets.h b/Engine/Source/Developer/Profiler/Private/Widgets/SFiltersAndPresets.h
--- a/Engine/Source/Developer/Profiler/Private/Widgets/SFiltersAndPresets.h
+++ b/Engine/Source/Developer/Profiler/Private/Widgets/SFiltersAndPresets.h
@@ -297,6 +297,8 @@ protected:
 	/** Called when the filter and presets widget should be updated with the latest data. */
 	void ProfilerManager_OnRequestFilterAndPresetsUpdate();
 
+	void ProfilerManager_OnViewModeChanged( EProfilerViewMode::Type NewViewMode );
+
 	void UpdateGroupAndStatTree( const FProfilerSessionPtr InProfilerSession );
 	void CreateGroups();
 	void SortStats();
diff --git a/Engine/Source/Developer/Profiler/Private/Widgets/SProfilerGraphPanel.cpp b/Engine/Source/Developer/Profiler/Private/Widgets/SProfilerGraphPanel.cpp
--- a/Engine/Source/Developer/Profiler/Private/Widgets/SProfilerGraphPanel.cpp
+++ b/Engine/Source/Developer/Profiler/Private/Widgets/SProfilerGraphPanel.cpp
@@ -17,8 +17,9 @@ SProfilerGraphPanel::~SProfilerGraphPanel()
 	// Remove ourselves from the profiler manager.
 	if( FProfilerManager::Get().IsValid() )
 	{
-		FProfilerManager::Get()->OnTrackedStatChanged().RemoveAll(this);
-		FProfilerManager::Get()->OnSessionInstancesUpdated().RemoveAll(this);
+		FProfilerManager::Get()->OnTrackedStatChanged().RemoveAll( this );
+		FProfilerManager::Get()->OnSessionInstancesUpdated().RemoveAll( this );
+		FProfilerManager::Get()->OnViewModeChanged().RemoveAll( this );
 	}
 }
 
@@ -76,6 +77,8 @@ void SProfilerGraphPanel::Construct( const FArguments& InArgs )
 	// Register ourselves with the profiler manager.
 	FProfilerManager::Get()->OnTrackedStatChanged().AddSP( this, &SProfilerGraphPanel::ProfilerManager_OnTrackedStatChanged );
 	FProfilerManager::Get()->OnSessionInstancesUpdated().AddSP( this, &SProfilerGraphPanel::ProfilerManager_OnSessionInstancesUpdated );
+
+	FProfilerManager::Get()->OnViewModeChanged().AddSP( this, &SProfilerGraphPanel::ProfilerManager_OnViewModeChanged );
 }
 
 
@@ -94,7 +97,7 @@ void SProfilerGraphPanel::HorizontalScrollBar_OnUserScrolled( float ScrollOffset
 	GraphOffset = FMath::Trunc( ScrollOffset * NumDataPoints );
 	DataGraph->ScrollTo( GraphOffset );
 
-	FProfilerManager::Get()->GetProfilerWindow()->ProfilerMiniView->OnSelectionBoxMoved( GraphOffset, GraphOffset + NumVisiblePoints );
+	FProfilerManager::Get()->GetProfilerWindow()->ProfilerMiniView->OnSelectionBoxChanged( GraphOffset, GraphOffset + NumVisiblePoints );
 }
 
 
@@ -125,7 +128,7 @@ void SProfilerGraphPanel::OnDataGraphGraphOffsetChanged( int32 InGraphOffset )
 {
 	GraphOffset = InGraphOffset;
 	SetScrollBarState();
-	FProfilerManager::Get()->GetProfilerWindow()->ProfilerMiniView->OnSelectionBoxMoved( InGraphOffset, InGraphOffset + NumVisiblePoints );
+	FProfilerManager::Get()->GetProfilerWindow()->ProfilerMiniView->OnSelectionBoxChanged( InGraphOffset, InGraphOffset + NumVisiblePoints );
 }
 
 void SProfilerGraphPanel::OnMiniViewSelectionBoxChanged( int32 FrameStart, int32 FrameEnd )
@@ -140,13 +143,31 @@ void SProfilerGraphPanel::DataGraph_OnViewModeChanged( EDataGraphViewModes::Type
 	UpdateInternals();
 }
 
+
+void SProfilerGraphPanel::ProfilerManager_OnViewModeChanged( EProfilerViewMode::Type NewViewMode )
+{
+	if( NewViewMode == EProfilerViewMode::LineIndexBased )
+	{
+		HorizontalScrollBar->SetVisibility( EVisibility::Collapsed );
+		HorizontalScrollBar->SetEnabled( false );
+	}
+	else if( NewViewMode == EProfilerViewMode::ThreadViewTimeBased )
+	{
+		HorizontalScrollBar->SetVisibility( EVisibility::Visible );
+		HorizontalScrollBar->SetEnabled( true );
+	}
+
+	//DataGraph->SetViewMode( NewViewMode );
+}
+
+
 void SProfilerGraphPanel::UpdateInternals()
 {
 	NumVisiblePoints = DataGraph->GetNumVisiblePoints();
 	NumDataPoints = DataGraph->GetNumDataPoints();
 
 	SetScrollBarState();
-	FProfilerManager::Get()->GetProfilerWindow()->ProfilerMiniView->OnSelectionBoxMoved( GraphOffset, GraphOffset + NumVisiblePoints );
+	FProfilerManager::Get()->GetProfilerWindow()->ProfilerMiniView->OnSelectionBoxChanged( GraphOffset, GraphOffset + NumVisiblePoints );
 	 
 	if( FProfilerManager::Get()->IsLivePreview() )
 	{
diff --git a/Engine/Source/Developer/Profiler/Private/Widgets/SProfilerGraphPanel.h b/Engine/Source/Developer/Profiler/Private/Widgets/SProfilerGraphPanel.h
--- a/Engine/Source/Developer/Profiler/Private/Widgets/SProfilerGraphPanel.h
+++ b/Engine/Source/Developer/Profiler/Private/Widgets/SProfilerGraphPanel.h
@@ -50,6 +50,8 @@ protected:
 	/** Called when the frame offset has been changed in the data graph widget. */
 	void OnDataGraphGraphOffsetChanged( int32 InFrameOffset );
 
+	void ProfilerManager_OnViewModeChanged( EProfilerViewMode::Type NewViewMode );
+
 public:
 	/** Called when the selection box has been changed in the profiler mini-view widget. */
 	void OnMiniViewSelectionBoxChanged( int32 FrameStart, int32 FrameEnd );
diff --git a/Engine/Source/Developer/Profiler/Private/Widgets/SProfilerMiniView.cpp b/Engine/Source/Developer/Profiler/Private/Widgets/SProfilerMiniView.cpp
--- a/Engine/Source/Developer/Profiler/Private/Widgets/SProfilerMiniView.cpp
+++ b/Engine/Source/Developer/Profiler/Private/Widgets/SProfilerMiniView.cpp
@@ -586,7 +586,7 @@ void SProfilerMiniView::AddThreadTime( int32 InFrameIndex, const TMap<uint32, fl
 	StatMetadata = InStatMetaData;
 }
 
-void SProfilerMiniView::OnSelectionBoxMoved( int32 FrameStart, int32 FrameEnd )
+void SProfilerMiniView::OnSelectionBoxChanged( int32 FrameStart, int32 FrameEnd )
 {
 	const int32 MaxFrameIndex = FMath::Max( AllFrames.Num() - 1, 0 );
 	SelectionBoxFrameStart = FMath::Clamp( FrameStart, 0, MaxFrameIndex );
diff --git a/Engine/Source/Developer/Profiler/Private/Widgets/SProfilerMiniView.h b/Engine/Source/Developer/Profiler/Private/Widgets/SProfilerMiniView.h
--- a/Engine/Source/Developer/Profiler/Private/Widgets/SProfilerMiniView.h
+++ b/Engine/Source/Developer/Profiler/Private/Widgets/SProfilerMiniView.h
@@ -311,7 +311,7 @@ public:
 	 *	Moves the selection box to the location specified by FrameStart and FrameEnd. 
 	 *	Called by external widgets.
 	 */
-	void OnSelectionBoxMoved( int32 FrameStart, int32 FrameEnd );
+	void OnSelectionBoxChanged( int32 FrameStart, int32 FrameEnd );
 
 	void DrawText( const FString& Text, const FSlateFontInfo& FontInfo, FVector2D Position, const FColor& TextColor, const FColor& ShadowColor, FVector2D ShadowOffset ) const;
 
diff --git a/Engine/Source/Developer/Profiler/Private/Widgets/SProfilerWindow.cpp b/Engine/Source/Developer/Profiler/Private/Widgets/SProfilerWindow.cpp
--- a/Engine/Source/Developer/Profiler/Private/Widgets/SProfilerWindow.cpp
+++ b/Engine/Source/Developer/Profiler/Private/Widgets/SProfilerWindow.cpp
@@ -10,15 +10,6 @@
 
 #define LOCTEXT_NAMESPACE "SProfilerWindow"
 
-/** Contains names of all widgets used in the profiler window. */
-namespace FNames
-{
-	static const FName GraphView = FName("GraphViewTab");
-	static const FName SentinelView = FName("SentinelViewTab");
-	static const FName FiltersAndPresets = FName("FilterAndPreset");
-	static const FName EventGraphView = FName("EventGraphView");
-}
-
 static FText GetTextForNotification( const EProfilerNotificationTypes::Type NotificatonType, const ELoadingProgressStates::Type ProgressState, const FString& Filename, const float ProgressPercent = 0.0f )
 {
 	FText Result;
@@ -90,10 +81,8 @@ SProfilerWindow::~SProfilerWindow()
 {}
 
 BEGIN_SLATE_FUNCTION_BUILD_OPTIMIZATION
-void SProfilerWindow::Construct( const FArguments& InArgs, const ISessionManagerRef& InSessionManager )
+void SProfilerWindow::Construct( const FArguments& InArgs )
 {
-	SessionManager = InSessionManager;
-
 	// TODO: Cleanup overlay code.
 	ChildSlot
 		[
@@ -377,6 +366,7 @@ EVisibility SProfilerWindow::IsSessionOverlayVissible() const
 	return EVisibility::Visible;
 }
 
+
 bool SProfilerWindow::IsProfilerEnabled() const
 {
 	const bool bIsActive = FProfilerManager::Get()->IsConnected() || FProfilerManager::Get()->IsCaptureFileFullyProcessed();
diff --git a/Engine/Source/Developer/Profiler/Private/Widgets/SProfilerWindow.h b/Engine/Source/Developer/Profiler/Private/Widgets/SProfilerWindow.h
--- a/Engine/Source/Developer/Profiler/Private/Widgets/SProfilerWindow.h
+++ b/Engine/Source/Developer/Profiler/Private/Widgets/SProfilerWindow.h
@@ -36,11 +36,8 @@ public:
 	
 	/**
 	 * Constructs this widget.
-	 *
-	 * @param InArgs					- the Slate argument list
-	 * @param InSessionManager			- the session manager to use
 	 */
-	void Construct( const FArguments& InArgs, const ISessionManagerRef& InSessionManager );
+	void Construct( const FArguments& InArgs );
 
 	void ManageEventGraphTab( const FGuid ProfilerInstanceID, const bool bCreateFakeTab, const FString TabName );
 	void UpdateEventGraph( const FGuid ProfilerInstanceID, const FEventGraphDataRef AverageEventGraph, const FEventGraphDataRef MaximumEventGraph, bool bInitial );
@@ -98,9 +95,6 @@ private:
 	/** Holds all event graphs. */
 	TSharedPtr<SVerticalBox> EventGraphPanel;
 
-	/** Holds the session manager. */
-	ISessionManagerPtr SessionManager;
-
 	/** Widget for the panel which contains all graphs and event graphs. */
 	TSharedPtr<SProfilerGraphPanel> GraphPanel;
 
