Profiler - Removed handling the event OnViewModeChanged,  it's not needed here atm

diff --git a/Engine/Source/Developer/Profiler/Private/Widgets/SFiltersAndPresets.cpp b/Engine/Source/Developer/Profiler/Private/Widgets/SFiltersAndPresets.cpp
--- a/Engine/Source/Developer/Profiler/Private/Widgets/SFiltersAndPresets.cpp
+++ b/Engine/Source/Developer/Profiler/Private/Widgets/SFiltersAndPresets.cpp
@@ -582,7 +582,6 @@ SFiltersAndPresets::~SFiltersAndPresets()
 	if( FProfilerManager::Get().IsValid() )
 	{
 		FProfilerManager::Get()->OnRequestFilterAndPresetsUpdate().RemoveAll( this );
-		FProfilerManager::Get()->OnViewModeChanged().RemoveAll( this );
 	}
 }
 
@@ -743,8 +742,6 @@ void SFiltersAndPresets::Construct( const FArguments& InArgs )
 	// Register ourselves with the profiler manager.
 	FProfilerManager::Get()->OnRequestFilterAndPresetsUpdate().AddSP( this, &SFiltersAndPresets::ProfilerManager_OnRequestFilterAndPresetsUpdate );
 
-	FProfilerManager::Get()->OnViewModeChanged().AddSP( this, &SFiltersAndPresets::ProfilerManager_OnViewModeChanged );
-
 	// Create the search filters: text based, stat type based etc.
 	// @TODO: HandleItemToStringArray should be moved somewhere else
 	GroupAndStatTextFilter = MakeShareable( new FGroupAndStatTextFilter( FGroupAndStatTextFilter::FItemToStringArray::CreateSP( this, &SFiltersAndPresets::HandleItemToStringArray ) ) );
@@ -760,22 +757,6 @@ void SFiltersAndPresets::ProfilerManager_OnRequestFilterAndPresetsUpdate()
 	bUpdateOnNextTick = true;
 }
 
-
-void SFiltersAndPresets::ProfilerManager_OnViewModeChanged( EProfilerViewMode::Type NewViewMode )
-{
-	if( NewViewMode == EProfilerViewMode::LineIndexBased )
-	{
-		SetVisibility( EVisibility::Visible );
-		SetEnabled( true );
-	}
-	else if( NewViewMode == EProfilerViewMode::ThreadViewTimeBased )
-	{
-		SetVisibility( EVisibility::Collapsed );
-		SetEnabled( false );
-	}
-}
-
-
 void SFiltersAndPresets::Tick( const FGeometry& AllottedGeometry, const double InCurrentTime, const float InDeltaTime )
 {
 	if( bUpdateOnNextTick )
diff --git a/Engine/Source/Developer/Profiler/Private/Widgets/SFiltersAndPresets.h b/Engine/Source/Developer/Profiler/Private/Widgets/SFiltersAndPresets.h
--- a/Engine/Source/Developer/Profiler/Private/Widgets/SFiltersAndPresets.h
+++ b/Engine/Source/Developer/Profiler/Private/Widgets/SFiltersAndPresets.h
@@ -297,8 +297,6 @@ protected:
 	/** Called when the filter and presets widget should be updated with the latest data. */
 	void ProfilerManager_OnRequestFilterAndPresetsUpdate();
 
-	void ProfilerManager_OnViewModeChanged( EProfilerViewMode::Type NewViewMode );
-
 	void UpdateGroupAndStatTree( const FProfilerSessionPtr InProfilerSession );
 	void CreateGroups();
 	void SortStats();
