Profiler - Moved FBinaryFindIndex to a better place, fixed issue if used with the custom range

diff --git a/Engine/Source/Developer/Profiler/Private/ProfilerDataProvider.h b/Engine/Source/Developer/Profiler/Private/ProfilerDataProvider.h
--- a/Engine/Source/Developer/Profiler/Private/ProfilerDataProvider.h
+++ b/Engine/Source/Developer/Profiler/Private/ProfilerDataProvider.h
@@ -6,30 +6,6 @@
 	Declarations
 -----------------------------------------------------------------------------*/
 
-class FBinaryFindIndex
-{
-public:
-	template<class T> 
-	static int32 LessEqual( const TArray<T>& Array, const T& Item, const int32 FirstIndex = 0, const int32 LastIndex = INDEX_NONE )
-	{
-		const int32 LocalLastIndex = LastIndex == INDEX_NONE ? Array.Num() : LastIndex;
-		int32 MaxIndex = LocalLastIndex;
-		int32 MidIndex = LocalLastIndex;
-		int32 MinIndex = FMath::Min( FirstIndex, (MaxIndex-1) / 2 );
-
-		while( MidIndex > 0 )
-		{
-			MidIndex = MaxIndex / 2;
-			if( Array[ MinIndex + MidIndex ] <= Item ) 
-			{
-				MinIndex += MidIndex;
-			}
-			MaxIndex -= MidIndex;
-		}
-		return MinIndex;
-	}
-};
-
 /** Data provider interface, acts as a container for profiler samples. */
 class IDataProvider : public TSharedFromThis<IDataProvider>
 {
diff --git a/Engine/Source/Developer/Profiler/Public/ProfilerCommon.h b/Engine/Source/Developer/Profiler/Public/ProfilerCommon.h
--- a/Engine/Source/Developer/Profiler/Public/ProfilerCommon.h
+++ b/Engine/Source/Developer/Profiler/Public/ProfilerCommon.h
@@ -114,4 +114,56 @@ public:
 		}
 		return Result;
 	}
-};
\ No newline at end of file
+};
+
+/** Implements binary search for the various cases. */
+class FBinaryFindIndex
+{
+public:
+	template<class T>
+	static int32 LessEqual( const TArray<T>& Array, const T& Item, const int32 FirstIndex = 0, const int32 LastIndex = INDEX_NONE )
+	{
+		const int32 LocalLastIndex = LastIndex == INDEX_NONE ? Array.Num() : LastIndex;
+		int32 Lenght = LocalLastIndex - FirstIndex;
+		int32 Middle = Lenght / 2;
+		int32 Offset = FirstIndex;
+
+		while( Middle > 0 )
+		{
+			Middle = Lenght / 2;
+			if( Array[Offset + Middle] <= Item )
+			{
+				Offset += Middle;
+			}
+			Lenght -= Middle;
+		}
+		return Offset;
+	}
+
+	template<class T>
+	static int32 GreaterEqual( const TArray<T>& Array, const T& Item, const int32 FirstIndex = 0, const int32 LastIndex = INDEX_NONE )
+	{
+		const int32 LocalLastIndex = LastIndex == INDEX_NONE ? Array.Num() : LastIndex;
+		int32 Lenght = LocalLastIndex - FirstIndex;
+		int32 Middle = Lenght / 2;
+		int32 Offset = FirstIndex;
+		int32 Edge = 0;
+
+		while( Middle > 0 )
+		{
+			Middle = Lenght / 2;
+			if( Array[Offset + Middle] >= Item )
+			{
+				Edge = 0;
+			}
+			else
+			{
+				Edge = 1;
+				Offset += Middle;		
+			}
+
+			Lenght -= Middle;
+		}
+		return Offset+Edge;
+	}
+};
