Profiler - Added first version of the thread view (to capture the data in the game use 'stat startfileraw')

diff --git a/Engine/Source/Developer/Profiler/Private/ProfilerManager.cpp b/Engine/Source/Developer/Profiler/Private/ProfilerManager.cpp
--- a/Engine/Source/Developer/Profiler/Private/ProfilerManager.cpp
+++ b/Engine/Source/Developer/Profiler/Private/ProfilerManager.cpp
@@ -232,6 +232,7 @@ void FProfilerManager::LoadRawStatsFile( const FString& RawStatsFileFileath )
 	ProfilerSession->PrepareLoading();
 	RequestFilterAndPresetsUpdateEvent.Broadcast();
 	//ProfilerSession_OnCaptureFileProcessed( ProfilerInstanceID );
+	bHasCaptureFileFullyProcessed = true;
 	//TrackDefaultStats();
 
 	SessionInstancesUpdatedEvent.Broadcast();
@@ -239,6 +240,8 @@ void FProfilerManager::LoadRawStatsFile( const FString& RawStatsFileFileath )
 	
 	GetProfilerWindow()->ManageEventGraphTab( ProfilerInstanceID, true, ProfilerSession->GetName() );
 	SetViewMode( EProfilerViewMode::ThreadViewTimeBased );
+
+	GetProfilerWindow()->GraphPanel->ThreadView->AttachProfilerStream( ProfilerSession->GetStream() );
 }
 
 
@@ -406,6 +409,7 @@ void FProfilerManager::SessionManager_OnInstanceSelectionChanged()
 			ActiveSession = SelectedSession;
 			ProfilerClient->Subscribe( ActiveSession->GetSessionId() );
 			ProfilerType = EProfilerSessionTypes::Live;
+			SetViewMode( EProfilerViewMode::LineIndexBased );
 		}
 		else
 		{
diff --git a/Engine/Source/Developer/Profiler/Private/ProfilerPrivatePCH.h b/Engine/Source/Developer/Profiler/Private/ProfilerPrivatePCH.h
--- a/Engine/Source/Developer/Profiler/Private/ProfilerPrivatePCH.h
+++ b/Engine/Source/Developer/Profiler/Private/ProfilerPrivatePCH.h
@@ -44,8 +44,9 @@
 #include "SFiltersAndPresets.h"
 #include "SEventGraph.h"
 #include "SDataGraph.h"
-#include "SProfilerGraphPanel.h"
 #include "SProfilerMiniView.h"
+#include "SProfilerThreadView.h"
+#include "SProfilerGraphPanel.h"
 #include "SProfilerFPSChartPanel.h"
 #include "SProfilerWindow.h"
 #include "SHistogram.h"
diff --git a/Engine/Source/Developer/Profiler/Private/ProfilerSession.cpp b/Engine/Source/Developer/Profiler/Private/ProfilerSession.cpp
--- a/Engine/Source/Developer/Profiler/Private/ProfilerSession.cpp
+++ b/Engine/Source/Developer/Profiler/Private/ProfilerSession.cpp
@@ -450,6 +450,7 @@ FRawProfilerSession::~FRawProfilerSession()
 
 bool FRawProfilerSession::HandleTicker( float DeltaTime )
 {
+#if	0
 	StatsThreadStats;
 	Stream;
 
@@ -481,6 +482,7 @@ bool FRawProfilerSession::HandleTicker( float DeltaTime )
 			break;
 		}
 	}
+#endif // 0
 
 	return true;
 }
@@ -523,25 +525,60 @@ void FRawProfilerSession::PrepareLoading()
 
 		// Update profiler's metadata.
 		StatMetaData->UpdateFromStatsState( StatsThreadStats );
+		const uint32 GameThreadID = GetMetaData()->GetGameThreadID();
 
 		// Read frames offsets.
 		Stream.ReadFramesOffsets( *FileReader );
 
-		// Verify frames offsets.
+		// Prepare profiler frames.
+		double ElapsedTimeMS = 0;
 		for( int32 FrameIndex = 0; FrameIndex < Stream.FramesInfo.Num(); ++FrameIndex )
 		{
 			const FStatsFrameInfo& StatsFrameInfo = Stream.FramesInfo[FrameIndex];
-			FileReader->Seek( StatsFrameInfo.FrameOffset );
+			FileReader->Seek( StatsFrameInfo.FrameFileOffset );
 			int64 TargetFrame;
 			*FileReader << TargetFrame;
-			TargetFrame = TargetFrame;
+
+			const double GameThreadTimeMS = GetMetaData()->ConvertCyclesToMS( (uint32)StatsFrameInfo.ThreadCycles.FindChecked( GameThreadID ) );
+
+			// Update mini-view, convert from cycles to ms.
+			TMap<uint32, float> ThreadTimesMS;
+			for( auto InnerIt = StatsFrameInfo.ThreadCycles.CreateConstIterator(); InnerIt; ++InnerIt )
+			{
+				const uint32 ThreadID = InnerIt.Key();
+				const bool bCanBeAdded = GameThreadID == ThreadID || GetMetaData()->GetRenderThreadID().Contains( ThreadID );
+				if( bCanBeAdded )
+				{
+					ThreadTimesMS.Add( ThreadID, GetMetaData()->ConvertCyclesToMS( InnerIt.Value() ) );
+				}
+			}
+
+			// Pass the reference to the stats' metadata.
+			OnAddThreadTime.ExecuteIfBound( FrameIndex, ThreadTimesMS, StatMetaData );
+
+			// Create a new profiler frame and add it to the stream.
+			ElapsedTimeMS += GameThreadTimeMS;
+			FProfilerFrame* ProfilerFrame = new FProfilerFrame( TargetFrame, GameThreadTimeMS, ElapsedTimeMS );
+			ProfilerFrame->ThreadTimesMS = ThreadTimesMS;
+			ProfilerStream.AddProfilerFrame( TargetFrame, ProfilerFrame );
+
+			// Thread cycles can be used to generate low-res thread view.
+			// @TODO yrx 2014-04-23 
 		}
-		int64 FrameOffset0 = Stream.FramesInfo[0].FrameOffset;
+		int64 FrameOffset0 = Stream.FramesInfo[0].FrameFileOffset;
 		FileReader->Seek( FrameOffset0 );
+
 #if	1
+		double CycleCounterAdjustmentMS = 0.0f;
+
 		// Read the raw stats messages.
 		for( int32 FrameIndex = 0; FrameIndex < Stream.FramesInfo.Num(); ++FrameIndex )
 		{
+			FProfilerFrame* ProfilerFrame = ProfilerStream.GetProfilerFrame( FrameIndex );
+			FileReader->Seek( Stream.FramesInfo[FrameIndex].FrameFileOffset );
+
+			UE_LOG( LogStats, Log, TEXT( "Processing raw stats frame: %4i/%4i" ), FrameIndex, Stream.FramesInfo.Num() );
+
 			FStatPacketArray Packet;
 
 			int64 TargetFrame;
@@ -557,26 +594,25 @@ void FRawProfilerSession::PrepareLoading()
 				Packet.Packets.Add( ToRead );
 			}
 
-			FProfilerFrame* ProfilerFrame = ProcessStatPacketArray( Packet );
+			ProcessStatPacketArray( Packet, *ProfilerFrame, FrameIndex ); // or ProfilerFrame->TargetFrame
+
+			// Find the first cycle counter for the game thread.
+			if( CycleCounterAdjustmentMS == 0.0f )
+			{
+				CycleCounterAdjustmentMS = ProfilerFrame->Root->CycleCounterStartTimeMS;
+			}
 
 			TMap<uint32, int64> ThreadCycles;
 			// Serialize thread cycles.
 			*FileReader << ThreadCycles;
-			
-			// Update mini-view.
-			// Convert from cycles to ms.
-			TMap<uint32, float> ThreadTimesMS;
-			for( auto InnerIt = ThreadCycles.CreateConstIterator(); InnerIt; ++InnerIt )
-			{
-				ThreadTimesMS.Add( InnerIt.Key(), StatMetaData->ConvertCyclesToMS( InnerIt.Value() ) );
-			}
+		
+			// Update thread time and mark profiler frame as valid and ready for use.
+			ProfilerFrame->MarkAsValid();
+		}
 
-			// Pass the reference to the stats' metadata.
-			OnAddThreadTime.ExecuteIfBound( FrameIndex, ThreadTimesMS, StatMetaData );
+		// Adjust all profiler frames.
+		ProfilerStream.AdjustCycleCounters( CycleCounterAdjustmentMS );
 
-			ProfilerFrame->ThreadTimesMS = ThreadTimesMS;
-			ProfilerStream.AddProfilerFrame( ProfilerFrame );
-		}
 #endif // 0
 	}
 
@@ -694,13 +730,19 @@ void FProfilerStatMetaData::UpdateFromStatsState( const FStatsThreadState& Stats
 			}
 			ThreadIDtoStatID.Add( ThreadID, StatID );
 
+			// Game thread is always NAME_GameThread
 			if( StatName == NAME_GameThread )
 			{
 				GameThreadID = ThreadID;
 			}
-			else if( StatName == NAME_RenderThread )
+			// Rendering thread may be "Rendering thread" or NAME_RenderThread with an index
+			else if( StatName.GetPlainNameString().Contains( FName(NAME_RenderThread).GetPlainNameString() ) )
+			{
+				RenderThreadIDs.AddUnique( ThreadID );
+			}
+			else if( StatName.GetPlainNameString().Contains( TEXT( "RenderingThread" ) ) )
 			{
-				RenderThreadIDs.Add( ThreadID );
+				RenderThreadIDs.AddUnique( ThreadID );
 			}
 		}
 
@@ -713,7 +755,7 @@ void FProfilerStatMetaData::UpdateFromStatsState( const FStatsThreadState& Stats
 	SecondsPerCycle = FPlatformTime::GetSecondsPerCycle();
 }
 
-FProfilerFrame* FRawProfilerSession::ProcessStatPacketArray( const FStatPacketArray& StatPacketArray )
+void FRawProfilerSession::ProcessStatPacketArray( const FStatPacketArray& StatPacketArray, FProfilerFrame& out_ProfilerFrame, int32 FrameIndex )
 {
 	// @TODO yrx 2014-03-24 Standardize thread names and id
 	// @TODO yrx 2014-04-22 Remove all references to the data provider, event graph etc once data graph can visualize.
@@ -733,15 +775,21 @@ FProfilerFrame* FRawProfilerSession::ProcessStatPacketArray( const FStatPacketAr
 	for( int32 PacketIndex = 0; PacketIndex < StatPacketArray.Packets.Num(); PacketIndex++ )
 	{
 		const FStatPacket& StatPacket = *StatPacketArray.Packets[PacketIndex];
-		const FName ThreadFName = StatsThreadStats.Threads.FindChecked( StatPacket.ThreadId );
+		FName ThreadFName = StatsThreadStats.Threads.FindChecked( StatPacket.ThreadId );
 		const uint32 NewThreadID = MetaData->ThreadIDtoStatID.FindChecked( StatPacket.ThreadId );
 
-		// @TODO yrx 2014-04-18 At this moment only for the game thread.
-		if( StatPacket.ThreadType != EThreadType::Game )
+		// @TODO yrx 2014-04-29 Only game or render thread is supported at this moment.
+		if( StatPacket.ThreadType != EThreadType::Game && StatPacket.ThreadType != EThreadType::Renderer )
 		{
 			continue;
 		}
 
+		// Workaround for issue with rendering thread names.
+		if( StatPacket.ThreadType == EThreadType::Renderer )
+		{
+			ThreadFName = NAME_RenderThread;
+		}
+
 		FProfilerStackNode* ThreadNode = ThreadNodes.FindRef( ThreadFName );
 		if( !ThreadNode )
 		{
@@ -763,7 +811,7 @@ FProfilerFrame* FRawProfilerSession::ProcessStatPacketArray( const FStatPacketAr
 				FrameRootSampleIndex
 			);
 
-			ThreadNode = ThreadNodes.Add( ThreadFName, new FProfilerStackNode( nullptr, ThreadMessage, ThreadRootSampleIndex ) );
+			ThreadNode = ThreadNodes.Add( ThreadFName, new FProfilerStackNode( nullptr, ThreadMessage, ThreadRootSampleIndex, FrameIndex ) );
 		}
 
 		if( StatPacket.ThreadType == EThreadType::Game )
@@ -791,7 +839,7 @@ FProfilerFrame* FRawProfilerSession::ProcessStatPacketArray( const FStatPacketAr
 				check( Item.NameAndInfo.GetFlag( EStatMetaFlags::IsCycle ) );
 				if( Op == EStatOperation::CycleScopeStart )
 				{
-					FProfilerStackNode* ChildNode = new FProfilerStackNode( Current, Item, -1 );
+					FProfilerStackNode* ChildNode = new FProfilerStackNode( Current, Item, -1, FrameIndex );
 					Current->Children.Add( ChildNode );
 
 					// Add a child sample.
@@ -818,9 +866,9 @@ FProfilerFrame* FRawProfilerSession::ProcessStatPacketArray( const FStatPacketAr
 					const int64 Delta = int32( uint32( ScopeEnd.GetValue_int64() ) - uint32( ScopeStart.GetValue_int64() ) );
 					Current->CyclesEnd = Current->CyclesStart + int64( Delta );
 
-					Current->TimeMSStart = MetaData->ConvertCyclesToMS( Current->CyclesStart );
-					Current->TimeMSEnd = MetaData->ConvertCyclesToMS( Current->CyclesEnd );
-					check( Current->TimeMSEnd >= Current->TimeMSStart );
+					Current->CycleCounterStartTimeMS = MetaData->ConvertCyclesToMS( Current->CyclesStart );
+					Current->CycleCounterEndTimeMS = MetaData->ConvertCyclesToMS( Current->CyclesEnd );
+					check( Current->CycleCounterEndTimeMS >= Current->CycleCounterStartTimeMS );
 
 					FProfilerStackNode* ChildNode = Current;
 
@@ -844,8 +892,8 @@ FProfilerFrame* FRawProfilerSession::ProcessStatPacketArray( const FStatPacketAr
 			const int32 LastChildIndex = ThreadNode.Children.Num() - 1;
 			ThreadNode.CyclesStart = ThreadNode.Children[0]->CyclesStart;
 			ThreadNode.CyclesEnd = ThreadNode.Children[LastChildIndex]->CyclesEnd;
-			ThreadNode.TimeMSStart = MetaData->ConvertCyclesToMS( ThreadNode.CyclesStart );
-			ThreadNode.TimeMSEnd = MetaData->ConvertCyclesToMS( ThreadNode.CyclesEnd );
+			ThreadNode.CycleCounterStartTimeMS = MetaData->ConvertCyclesToMS( ThreadNode.CyclesStart );
+			ThreadNode.CycleCounterEndTimeMS = MetaData->ConvertCyclesToMS( ThreadNode.CyclesEnd );
 
 			FProfilerSample& ProfilerSample = MutableCollection[ThreadNode.SampleIndex];
 			ProfilerSample.SetStartAndEndMS( MetaData->ConvertCyclesToMS( ThreadNode.CyclesStart ), MetaData->ConvertCyclesToMS( ThreadNode.CyclesEnd ) );
@@ -869,12 +917,14 @@ FProfilerFrame* FRawProfilerSession::ProcessStatPacketArray( const FStatPacketAr
 	// Update aggregated events.
 	UpdateAggregatedEventGraphData( LastFrameIndex );
 
-	FProfilerFrame* ProfilerFrame = new FProfilerFrame();
+	// RootNode is the same as the game thread node.
+	out_ProfilerFrame.Root->CycleCounterStartTimeMS = GameThreadStartMS;
+	out_ProfilerFrame.Root->CycleCounterEndTimeMS = GameThreadEndMS;
+
 	for( auto It = ThreadNodes.CreateIterator(); It; ++It )
 	{
-		ProfilerFrame->Root->Children.Add( It.Value() );
-		// @TODO yrx 2014-04-22 Sort children GameThread, RenderThread, Other children sorted by total time.
+		out_ProfilerFrame.AddChild( It.Value() );
 	}
 
-	return ProfilerFrame;
+	out_ProfilerFrame.SortChildren();
 }
\ No newline at end of file
diff --git a/Engine/Source/Developer/Profiler/Private/ProfilerSession.h b/Engine/Source/Developer/Profiler/Private/ProfilerSession.h
--- a/Engine/Source/Developer/Profiler/Private/ProfilerSession.h
+++ b/Engine/Source/Developer/Profiler/Private/ProfilerSession.h
@@ -497,11 +497,13 @@ private:
 					StatPtr->_Name = *UniqueThreadName;
 					ThreadIDtoStatID.Add( ThreadID, StatID );
 
+					// Game thread is always NAME_GameThread
 					if( **ThreadDesc == FName( NAME_GameThread ) )
 					{
 						GameThreadID = ThreadID;
 					}
-					else if( **ThreadDesc == FName( NAME_RenderThread ) )
+					// Rendering thread may be "Rendering thread" or NAME_RenderThread with an index
+					else if( ThreadDesc->Contains( FName( NAME_RenderThread ).GetPlainNameString() ) )
 					{
 						RenderThreadIDs.AddUnique( ThreadID );
 					}
@@ -1297,9 +1299,14 @@ public:
 	/** Starts a process of loading the raw stats file. */
 	void PrepareLoading();
 
+	const FProfilerStream& GetStream() const
+	{
+		return ProfilerStream;
+	}
+
 	/**
 	 *	Process all stats packets and convert them to data accessible by the profiler.
 	 *	Temporary version, will be optimized later.
 	 */
-	FProfilerFrame* ProcessStatPacketArray( const FStatPacketArray& PacketArray );
+	void ProcessStatPacketArray( const FStatPacketArray& PacketArray, FProfilerFrame& out_ProfilerFrame, int32 FrameIndex );
 };
\ No newline at end of file
diff --git a/Engine/Source/Developer/Profiler/Private/ProfilerStream.cpp b/Engine/Source/Developer/Profiler/Private/ProfilerStream.cpp
--- a/Engine/Source/Developer/Profiler/Private/ProfilerStream.cpp
+++ b/Engine/Source/Developer/Profiler/Private/ProfilerStream.cpp
@@ -1,3 +1,178 @@
 // Copyright 1998-2014 Epic Games, Inc. All Rights Reserved.
 
-#include "ProfilerPrivatePCH.h"
\ No newline at end of file
+#include "ProfilerPrivatePCH.h"
+
+/*-----------------------------------------------------------------------------
+	FProfilerUIStream
+-----------------------------------------------------------------------------*/
+
+void FProfilerUIStream::GenerateUIStream( const FProfilerStream& ProfilerStream, const double StartTimeMS, const double EndTimeMS, const double ZoomFactorX, const double NumMillisecondsPerWindow, const double NumPixelsPerMillisecond, const double NumMillisecondsPerSample )
+{
+	SCOPE_LOG_TIME_FUNC();
+
+	ProfilerStream;
+	ZoomFactorX;
+
+	// Just copy the data for now.
+	const FIntPoint FramesIndices = ProfilerStream.GetFramesIndicesForTimeRange( StartTimeMS, EndTimeMS );
+	const double FrameTimeRangeMS = EndTimeMS - StartTimeMS;
+
+	const double WindowWidth = NumMillisecondsPerWindow*NumPixelsPerMillisecond;
+
+	// Preallocate nodes for threads.
+	const int32 NumThreads = ProfilerStream.GetNumThreads();
+	ThreadNodes.Reset( NumThreads );
+
+	if( true || NumMillisecondsPerWindow >= FrameTimeRangeMS )
+	{
+		// Just copy the data.
+		const int32 NumFrames = ProfilerStream.GetNumFrames();
+		const int32 MaxFrameIndex = FMath::Min( FramesIndices.Y + 1, NumFrames - 1 );
+
+		for( int32 FrameIndex = FramesIndices.X; FrameIndex < MaxFrameIndex; ++FrameIndex )
+		{
+			const FProfilerFrame* ProfilerFrame = ProfilerStream.GetProfilerFrame( FrameIndex );
+
+			ProfilerFrame->FrameTimeMS;
+			ProfilerFrame->LastAccessTime;
+			ProfilerFrame->ThreadTimesMS;
+
+			// Thread nodes.
+			const int32 NumThreads = ProfilerFrame->Root->Children.Num();
+			for( int32 ThreadIndex = 0; ThreadIndex < NumThreads; ++ThreadIndex )
+			{
+				FProfilerUIStackNode* ThreadUIStackNode = new FProfilerUIStackNode( ProfilerFrame->Root->Children[ThreadIndex], FProfilerUIStackNode::THREAD_NODE_INDEX, FProfilerUIStackNode::THREAD_NODE_INDEX, FrameIndex );
+				ThreadUIStackNode->InitializeUIData( NumMillisecondsPerWindow, NumPixelsPerMillisecond, NumMillisecondsPerSample );
+
+				ThreadNodes.AddRawItem( ThreadUIStackNode );
+
+				CombineOrSet( ThreadUIStackNode, *ProfilerFrame->Root->Children[ThreadIndex], ThreadIndex*DEFAULT_VISIBLE_THREAD_DEPTH, NumMillisecondsPerWindow, NumPixelsPerMillisecond, NumMillisecondsPerSample );
+			}
+		}
+	}
+	else
+	{
+		// @TODO yrx 2014-04-28 Combining profiler frames doesn't make sense at this moment
+		// At this scale we may want to switch into line graph or something?
+	}
+
+	// @TODO yrx 2014-04-25 MAX_THREAD_DEPTH;
+	LinearizeStream();
+}
+
+void FProfilerUIStream::CombineOrSet( FProfilerUIStackNode* ParentUIStackNode, const FProfilerStackNode& ProfilerStackNode, int32 GlobalNodeDepth, const double NumMillisecondsPerWindow, const double NumPixelsPerMillisecond, const double NumMillisecondsPerSample )
+{
+	const int32 ThreadIndex = GlobalNodeDepth / DEFAULT_VISIBLE_THREAD_DEPTH;
+	const int32 ThreadNodeDepth = GlobalNodeDepth % DEFAULT_VISIBLE_THREAD_DEPTH;
+	const int32 NumChildren = ProfilerStackNode.Children.Num();
+	const int32 FrameIndex = ProfilerStackNode.FrameIndex;
+
+	if( NumChildren > 0 )
+	{
+		FProfilerStackNode* FirstProfilerChildNode = nullptr;
+		bool bLookingForChildToCombine = false;
+
+		TArray<const FProfilerStackNode*> OriginalStackNodes;
+
+		for( int32 ChildIndex = 0; ChildIndex < NumChildren; ++ChildIndex )
+		{
+			if( bLookingForChildToCombine )
+			{
+				const FProfilerStackNode* NextProfilerChildNode = ProfilerStackNode.Children[ChildIndex];
+				const bool bNeedsToBeCombined = NextProfilerChildNode->GetDurationMS() < NumMillisecondsPerSample;
+
+				// The next node can be displayed in the UI, add it and the previous one as the combined. 
+				if( !bNeedsToBeCombined )
+				{
+					// Combine child nodes into one node.
+					{
+						FProfilerUIStackNode* ProfilerUIStackNode = new FProfilerUIStackNode( OriginalStackNodes, OriginalStackNodes.Num(), GlobalNodeDepth, ThreadIndex, FrameIndex );
+						ProfilerUIStackNode->InitializeUIData( NumMillisecondsPerWindow, NumPixelsPerMillisecond, NumMillisecondsPerSample );
+						ParentUIStackNode->Children.AddRawItem( ProfilerUIStackNode );
+						bLookingForChildToCombine = false;
+						OriginalStackNodes.Reset();
+					}
+
+					// Add the next node as a single node.
+					{
+						FProfilerUIStackNode* ProfilerUIStackNode = new FProfilerUIStackNode( NextProfilerChildNode, GlobalNodeDepth, ThreadIndex, FrameIndex );
+						ProfilerUIStackNode->InitializeUIData( NumMillisecondsPerWindow, NumPixelsPerMillisecond, NumMillisecondsPerSample );
+						ParentUIStackNode->Children.AddRawItem( ProfilerUIStackNode );
+
+						// Check the call stack depth.
+						if( ThreadNodeDepth != DEFAULT_VISIBLE_THREAD_DEPTH - 1 )
+						{
+							CombineOrSet( ProfilerUIStackNode, *NextProfilerChildNode, GlobalNodeDepth + 1, NumMillisecondsPerWindow, NumPixelsPerMillisecond, NumMillisecondsPerSample );
+						}
+						else
+						{
+							// Mark this node as culled.
+							ProfilerUIStackNode->MarkAsCulled();
+						}
+					}
+
+					continue;
+				}
+
+				// We need to combine a few child nodes into one node.
+				const double CombinedDurationMS = NextProfilerChildNode->CycleCounterEndTimeMS - FirstProfilerChildNode->CycleCounterStartTimeMS;
+
+				
+
+				OriginalStackNodes.Add( NextProfilerChildNode );
+				const bool bCanBeCombined = CombinedDurationMS > NumMillisecondsPerSample;
+				if( bCanBeCombined )
+				{
+					// Combine child nodes into one node.
+					FProfilerUIStackNode* ProfilerUIStackNode = new FProfilerUIStackNode( OriginalStackNodes, OriginalStackNodes.Num(), GlobalNodeDepth, ThreadIndex, FrameIndex );
+					ProfilerUIStackNode->InitializeUIData( NumMillisecondsPerWindow, NumPixelsPerMillisecond, NumMillisecondsPerSample );
+					ParentUIStackNode->Children.AddRawItem( ProfilerUIStackNode );
+					bLookingForChildToCombine = false;
+					OriginalStackNodes.Reset();
+				}
+
+				continue;
+			}
+
+			FirstProfilerChildNode = ProfilerStackNode.Children[ChildIndex];
+			const bool bNeedsToBeCombined = FirstProfilerChildNode->GetDurationMS() < NumMillisecondsPerSample;
+
+			if( !bNeedsToBeCombined && !bLookingForChildToCombine )
+			{
+				// We have a sample that can be displayed in the UI, add it.
+				FProfilerUIStackNode* ProfilerUIStackNode = new FProfilerUIStackNode( FirstProfilerChildNode, GlobalNodeDepth, ThreadIndex, FrameIndex );
+				ProfilerUIStackNode->InitializeUIData( NumMillisecondsPerWindow, NumPixelsPerMillisecond, NumMillisecondsPerSample );
+				ParentUIStackNode->Children.AddRawItem( ProfilerUIStackNode );
+
+				// Check the call stack depth.
+				if( ThreadNodeDepth != DEFAULT_VISIBLE_THREAD_DEPTH - 1 )
+				{
+					CombineOrSet( ProfilerUIStackNode, *FirstProfilerChildNode, GlobalNodeDepth + 1, NumMillisecondsPerWindow, NumPixelsPerMillisecond, NumMillisecondsPerSample );
+				}
+				else
+				{
+					// Mark this node as culled.
+					ProfilerUIStackNode->MarkAsCulled();
+				}
+			}
+			else
+			{
+				bLookingForChildToCombine = true;			
+				OriginalStackNodes.Add( FirstProfilerChildNode );
+			}
+		}
+
+		// We are out of the loop, but still need to combine children even if not visible.
+		if( bLookingForChildToCombine )
+		{
+			// Combine child nodes into one node.
+			{
+				FProfilerUIStackNode* ProfilerUIStackNode = new FProfilerUIStackNode( OriginalStackNodes, OriginalStackNodes.Num(), GlobalNodeDepth, ThreadIndex, FrameIndex );
+				ProfilerUIStackNode->InitializeUIData( NumMillisecondsPerWindow, NumPixelsPerMillisecond, NumMillisecondsPerSample );
+				ParentUIStackNode->Children.AddRawItem( ProfilerUIStackNode );
+				bLookingForChildToCombine = false;
+				OriginalStackNodes.Reset();
+			}
+		}
+	}
+}
diff --git a/Engine/Source/Developer/Profiler/Private/ProfilerStream.h b/Engine/Source/Developer/Profiler/Private/ProfilerStream.h
--- a/Engine/Source/Developer/Profiler/Private/ProfilerStream.h
+++ b/Engine/Source/Developer/Profiler/Private/ProfilerStream.h
@@ -9,44 +9,50 @@
 /** Profiler stack node, used to store the whole call stack for one frame. */
 struct FProfilerStackNode
 {
-	FProfilerStackNode* Parent;
-	TArray<FProfilerStackNode*> Children;
-
 	/** Short name. */
 	FName StatName;
 	FName LongName;
 
+	TArray<FProfilerStackNode*> Children;
+	FProfilerStackNode* Parent;
+
 	int64 CyclesStart;
 	int64 CyclesEnd;
 
-	double TimeMSStart;
-	double TimeMSEnd;
+	// @TODO yrx 2014-04-23 This should be based on the global time.
+	double CycleCounterStartTimeMS;
+	double CycleCounterEndTimeMS;
 
 	/** Index of this node in the data provider's collection. */
-	uint32 SampleIndex;
+	/*const*/ uint32 SampleIndex;
+
+	/** Index of the frame the this node belongs to. */
+	const int32 FrameIndex;
 
 	/** Initializes thread root node. */
-	FProfilerStackNode()
-		: Parent( nullptr )
-		, StatName( FStatConstants::NAME_ThreadRoot )
+	FProfilerStackNode( int32 InFrameIndex )
+		: StatName( FStatConstants::NAME_ThreadRoot )
 		, LongName( FStatConstants::NAME_ThreadRoot )
+		, Parent( nullptr )
 		, CyclesStart( 0 )
 		, CyclesEnd( 0 )
-		, TimeMSStart( 0.0f )
-		, TimeMSEnd( 0.0f )
+		, CycleCounterStartTimeMS( 0.0f )
+		, CycleCounterEndTimeMS( 0.0f )
 		, SampleIndex( 0 )
+		, FrameIndex( InFrameIndex )
 	{}
 
 	/** Initializes children node. */
-	FProfilerStackNode( FProfilerStackNode* InParent, const FStatMessage& StatMessage, uint32 InSampleIndex )
-		: Parent( InParent )
-		, StatName( StatMessage.NameAndInfo.GetShortName() )
+	FProfilerStackNode( FProfilerStackNode* InParent, const FStatMessage& StatMessage, uint32 InSampleIndex, int32 InFrameIndex )
+		: StatName( StatMessage.NameAndInfo.GetShortName() )
 		, LongName( StatMessage.NameAndInfo.GetRawName() )
+		, Parent( InParent )
 		, CyclesStart( StatMessage.GetValue_int64() )
 		, CyclesEnd( 0 )
-		, TimeMSStart( 0.0f )
-		, TimeMSEnd( 0.0f )
+		, CycleCounterStartTimeMS( 0.0f )
+		, CycleCounterEndTimeMS( 0.0f )
 		, SampleIndex( InSampleIndex )
+		, FrameIndex( InFrameIndex )
 	{}
 
 	~FProfilerStackNode()
@@ -57,6 +63,11 @@ struct FProfilerStackNode
 		}
 	}
 
+	double GetDurationMS() const
+	{
+		return CycleCounterEndTimeMS - CycleCounterStartTimeMS;
+	}
+
 	/** Calculates allocated size by all stacks node, only valid for the thread root node. */
 	int64 GetAllocatedSize() const
 	{
@@ -67,34 +78,67 @@ struct FProfilerStackNode
 		}
 		return AllocatedSize;
 	}
-};
 
-enum
-{
-	STACKNODE_INVALID = 0,
-	STACKNODE_VALID = 1,
+	void AdjustCycleCounters( double CycleCounterAdjustmentMS )
+	{
+		CycleCounterStartTimeMS -= CycleCounterAdjustmentMS;
+		CycleCounterEndTimeMS -= CycleCounterAdjustmentMS;
+		for( const auto& Child : Children )
+		{
+			Child->AdjustCycleCounters( CycleCounterAdjustmentMS );
+		}
+	}
 };
 
+
 /** Profiler frame. */
 struct FProfilerFrame
 {
+protected:
+	enum
+	{
+		STACKNODE_INVALID = 0,
+		STACKNODE_VALID = 1,
+	};
+
+public:
 	/** Root node. */
+	// @TODO yrx 2014-04-25 Replace with TIndirectArray<FProfilerStackNode>??
 	FProfilerStackNode* Root;
 
 	/** Thread times in milliseconds for this frame. */
 	TMap<uint32, float> ThreadTimesMS;
 
+	/** Target frame as captured by the stats system. */
+	int64 TargetFrame;
+
+	/** Frame time for this frame. */
+	double FrameTimeMS;
+
+	/** How many milliseconds have passed from the beginning. */
+	double ElapsedTimeMS;
+
+	/**
+	 *	Last time this frame has been accessed.
+	 *	Used by the profiler's GC to remove 'idle' profiler frames.
+	 *	Used only if working under specified memory constraint.
+	 */
+	double LastAccessTime;
+
 	// @TODO yrx 2014-04-22 Non-frame stats like accumulator, counters, memory etc?
 
 	/**
-	*	Indicates whether this node is in the memory, move to thread root node later.
-	*	Valid only for the root stack node.
-	*	This set by one thread and accessed by another thread, there is no thread contention.
+	*	Indicates whether this profiler frame is in the memory.
+	*	This is set by one thread and accessed by another thread, there is no thread contention.
 	*/
 	FThreadSafeCounter AccessLock;
 
-	FProfilerFrame()
-		: Root( new FProfilerStackNode() )
+	FProfilerFrame( int64 InTargetFrame, double InFrameTimeMS, double InElapsedTimeMS )
+		: Root( new FProfilerStackNode( (int32)(InTargetFrame&MAX_int32) ) )
+		, TargetFrame( InTargetFrame )
+		, FrameTimeMS( InFrameTimeMS )
+		, ElapsedTimeMS( InElapsedTimeMS )
+		, LastAccessTime( FPlatformTime::Seconds() )
 		, AccessLock( STACKNODE_INVALID )
 	{}
 
@@ -103,6 +147,32 @@ struct FProfilerFrame
 		FreeMemory();
 	}
 
+	void AddChild( FProfilerStackNode* ProfilerStackNode )
+	{
+		Root->Children.Add( ProfilerStackNode );
+	}
+
+	void SortChildren()
+	{
+		/**
+		*	Sorts thread nodes to be in a particular order.
+		*	GameThread, RenderThread.
+		*/
+		struct FThreadSort
+		{
+			FORCEINLINE bool operator()( const FProfilerStackNode& A, const FProfilerStackNode& B ) const
+			{
+				if( A.StatName == NAME_GameThread && B.StatName == NAME_RenderThread )
+				{
+					return true;
+				}
+				return false;
+			}
+		};
+
+		Root->Children.Sort( FThreadSort() );
+	}
+
 	void MarkAsValid()
 	{
 		const int32 OldLock = AccessLock.Set( STACKNODE_VALID );
@@ -125,6 +195,7 @@ struct FProfilerFrame
 		return ThreadTimesMS.GetAllocatedSize() + Root ? Root->GetAllocatedSize() : 0;
 	}
 
+	/** Frees most of the memory allocated by this profiler frame. */
 	void FreeMemory()
 	{
 		delete Root;
@@ -142,17 +213,35 @@ struct FProfilerStream
 protected:
 
 	/** History frames collected so far or read from the file. */
+	// @TODO yrx 2014-04-24 TSharedRef<THREAD_SAFE> ?
 	TChunkedArray<FProfilerFrame*> Frames;
 
+	/** Each element in this array stores the frame time, accessed by a frame index, in millisecond. */
+	TArray<double> FrameTimesMS;
+
+	/** Each element in this array stores the total frame time, accessed by a frame index, in millisecond. */
+	TArray<double> ElapsedFrameTimesMS;
+
+	/** Thread FNames that have been visible so far. */
+	TSet<FName> ThreadIDs;
+	
 	/** Critical section. */
 	mutable FCriticalSection CriticalSection;
 
 
 public:
-	void AddProfilerFrame( FProfilerFrame* ProfilerFrame )
+	void AddProfilerFrame( int64 TargetFrame, FProfilerFrame* ProfilerFrame )
 	{
 		FScopeLock Lock( &CriticalSection );
 		Frames.AddElement( ProfilerFrame );
+
+		FrameTimesMS.Add( ProfilerFrame->FrameTimeMS );
+		ElapsedFrameTimesMS.Add( ProfilerFrame->ElapsedTimeMS );
+
+		//TArray<uint32> FrameThreadIDs;
+		//ProfilerFrame->ThreadTimesMS.GenerateKeyArray( FrameThreadIDs );
+		ThreadIDs.Add( NAME_GameThread );
+		ThreadIDs.Add( NAME_RenderThread );
 	}
 
 	/**
@@ -164,6 +253,20 @@ public:
 		return Frames[FrameIndex];
 	}
 
+	/**
+	 * @return frames indices, where X is the start frame index, Y is the end frame index
+	 */
+	const FIntPoint GetFramesIndicesForTimeRange( const double StartTimeMS, const double EndTimeMS ) const
+	{
+		FScopeLock Lock( &CriticalSection );
+
+		// Find the start frame index when the start time is less or equal.
+		const int32 StartFrameIndex = FBinaryFindIndex::LessEqual( ElapsedFrameTimesMS, StartTimeMS );
+		const int32 EndFrameIndex = FBinaryFindIndex::GreaterEqual( ElapsedFrameTimesMS, EndTimeMS, StartFrameIndex );
+		
+		return FIntPoint( StartFrameIndex, EndFrameIndex );
+	}
+
 	int64 GetAllocatedSize() const
 	{
 		FScopeLock Lock( &CriticalSection );
@@ -180,31 +283,259 @@ public:
 
 		return AllocatedSize;
 	}
+
+	/** @return number of frames that have been collected so far. */
+	int32 GetNumFrames() const
+	{
+		FScopeLock Lock( &CriticalSection );
+		return Frames.Num();
+	}
+
+	/** @return the elapsed time for all collected frames. */
+	double GetElapsedTime() const
+	{
+		FScopeLock Lock( &CriticalSection );
+		return ElapsedFrameTimesMS[ElapsedFrameTimesMS.Num() - 1];
+	}
+
+	/** @return the frame duration for the specified frame, in milliseconds. */
+	const double GetFrameTimeMS( const int32 InFrameIndex ) const
+	{
+		FScopeLock Lock( &CriticalSection );
+		return FrameTimesMS[InFrameIndex];
+	}
+
+	/** @return the elapsed time for the specified frame, in milliseconds. */
+	const double GetElapsedFrameTimeMS( const int32 InFrameIndex ) const
+	{
+		FScopeLock Lock( &CriticalSection );
+		return ElapsedFrameTimesMS[InFrameIndex];
+	}
+
+	void AdjustCycleCounters( double CycleCounterAdjustmentMS )
+	{
+		FScopeLock Lock( &CriticalSection );
+		for( int32 FrameIndex = 0; FrameIndex < Frames.Num(); ++FrameIndex )
+		{
+			FProfilerStackNode* RootNode = Frames[FrameIndex]->Root;
+			RootNode->AdjustCycleCounters( CycleCounterAdjustmentMS );
+		}
+	}
+
+	/**
+	 * @return peak number of threads
+	 */
+	int32 GetNumThreads() const
+	{
+		FScopeLock Lock( &CriticalSection );
+		return ThreadIDs.Num();
+	}
 };
 
+// @TODO yrx 2014-04-23 ProfilerFrameLOD?
+
 /**
 *	Profiler UI stack node.
 *	Similar to the profiler stack node, but contains data prepared and optimized for the UI
 */
 struct FProfilerUIStackNode
 {
-	/** Frames used to generate this data, at least 1 frame. */
-	TArray<int64> Frames;
+	enum 
+	{
+		/** Indicates that the nodes is the thread node and shouldn't be displayed. */
+		THREAD_NODE_INDEX = -1,
+	};
+
+// 	FProfilerUIStackNode()
+// 		: CycleCountersStartTimeMS( 0.0f )
+// 		, CycleCountersEndTimeMS( 0.0f )
+// 		, PositionXPx( 0.0f )
+// 		, PositionY( 0.0f )
+// 		, Width( 0.0f )
+// 		, GlobalNodeDepth( THREAD_NODE_INDEX )
+// 		, ThreadNodeDepth( 0 )
+// 		, ThreadIndex( THREAD_NODE_INDEX )
+// 		, bIsCombined( false )
+// 	{}
+
+	FProfilerUIStackNode( const FProfilerStackNode* ProfilerStackNode, int32 InGlobalNodeDepth, int32 InThreadIndex, int32 InFrameIndex )
+		: StatName( ProfilerStackNode->StatName )
+		, LongName( ProfilerStackNode->LongName )
+		, CycleCountersStartTimeMS( ProfilerStackNode->CycleCounterStartTimeMS )
+		, CycleCountersEndTimeMS( ProfilerStackNode->CycleCounterEndTimeMS )
+		, PositionXPx( 0.0f )
+		, PositionY( 0.0f )
+		, WidthPx( 0.0f )
+		, GlobalNodeDepth( InGlobalNodeDepth )
+		, ThreadNodeDepth( 0 )
+		, ThreadIndex( InThreadIndex )
+		, FrameIndex( InFrameIndex )
+		, bIsCombined( false )
+		, bIsCulled( false )
+	{
+		OriginalStackNodes.Add( ProfilerStackNode );
+	}
+
+	FProfilerUIStackNode( const TArray<const FProfilerStackNode*>& ProfilerStackNodes, int32 NumStackNodes, int32 InGlobalNodeDepth, int32 InThreadIndex, int32 InFrameIndex )
+		: StatName( *FString::Printf( TEXT( "[%i]" ), NumStackNodes ) )
+		, LongName( *FString::Printf( TEXT( "Combined %i items" ), NumStackNodes ) )
+		, CycleCountersStartTimeMS( ProfilerStackNodes[0]->CycleCounterStartTimeMS )
+		, CycleCountersEndTimeMS( ProfilerStackNodes[NumStackNodes - 1]->CycleCounterEndTimeMS )
+		, PositionXPx( 0.0f )
+		, PositionY( 0.0f )
+		, WidthPx( 0.0f )
+		, GlobalNodeDepth( InGlobalNodeDepth )
+		, ThreadNodeDepth( 0 )
+		, ThreadIndex( InThreadIndex )
+		, FrameIndex( InFrameIndex )
+		, bIsCombined( true )
+		, bIsCulled( false )
+	{
+		OriginalStackNodes.Append( ProfilerStackNodes );
+	}
 
-	FProfilerUIStackNode* Parent;
-	TArray<FProfilerUIStackNode*> Children;
+	~FProfilerUIStackNode()
+	{
+		OriginalStackNodes;
+	}
+	
+	void InitializeUIData( const double NumMillisecondsPerWindow, const double NumPixelsPerMillisecond, const double NumMillisecondsPerSample )
+	{
+		const double DurationMS = GetDurationMS();
+		const double NumPixels = DurationMS*NumPixelsPerMillisecond;
+		WidthPx = NumPixels;
+
+		PositionXPx = CycleCountersStartTimeMS*NumPixelsPerMillisecond;
+		PositionY = (double)GlobalNodeDepth;
+	}
+
+	void MarkAsCulled()
+	{
+		bIsCulled = true;
+	}
+
+	double GetDurationMS() const
+	{
+		return CycleCountersEndTimeMS - CycleCountersStartTimeMS;
+	}
+
+	FVector2D GetLocalPosition( double PositionXOffsetPx, double PositionYOffset ) const
+	{
+		return FVector2D( PositionXPx - PositionXOffsetPx, PositionY - PositionYOffset );
+	}
+
+	bool IsVisible() const
+	{
+		return true;
+	}
+
+	/**
+	 *	Original stack nodes used to generate this UI stack node.
+	 *	Useful if this node is a combined stack node.
+	 */
+	TArray<const FProfilerStackNode*> OriginalStackNodes;
+
+	TIndirectArray<FProfilerUIStackNode> Children;
 
 	/** Short name. */
 	FName StatName;
+
+	/** Long name, short name, stat description, group name. */
 	FName LongName;
 
-	double TimeMSStart;
-	double TimeMSEnd;
+	/** Time range of the stack node. */
+	double CycleCountersStartTimeMS;
+	double CycleCountersEndTimeMS;
+
+	/** UI Data. */
+
+	/** Position of the stack node, absolute position, needs to be converted to the local before rendering. */
+	double PositionXPx;
+	double PositionY;
+
+	/** Width of the stack node, in pixels. */
+	double WidthPx;
+
+	/** Depth of this node, in the global scope. */
+	int32 GlobalNodeDepth;
+
+	/** Depth of this node, in the thread scope. */
+	int32 ThreadNodeDepth;
+
+	/** Thread index of this node. */
+	int32 ThreadIndex;
+
+	/** Index of the frame the this node belongs to. */
+	int32 FrameIndex;
+
+	/**
+	 *	Whether this UI stack node has been combined, if this is true we stop processing here.
+	 *	Mouse processing is disabled, and there is no tooltip displayed.
+	 *	The user needs to zoom-in to see more details.
+	 */
+	bool bIsCombined;
+
+	/**
+	 *	True means that the UI stack node has deeper call stack, but can't be displayed due to UI limitation.
+	 *	Culled nodes are rendered with a special marker to indicate that there are more nodes.
+	 */
+	bool bIsCulled;
+
+	/** Access lock. */
+	FThreadSafeCounter AccessLock;
 };
 
 
 struct FProfilerUIStream
 {
-	TArray<FProfilerFrame*> Frames;
+	enum
+	{
+		/**
+		 *	Default number of rows displaying cycle counters.
+		 *	Read it as the call stack depth.
+		 */
+		// @TODO yrx 2014-04-25 This probably should be the max seen so far for each thread.
+		DEFAULT_VISIBLE_THREAD_DEPTH = 16,
+	};
+
+	TIndirectArray<FProfilerUIStackNode> ThreadNodes;
+
+	TArray< TArray<const FProfilerUIStackNode*> > LinearRowsOfNodes;
+
+	void GenerateUIStream( const FProfilerStream& ProfilerStream, const double StartTimeMS, const double EndTimeMS, const double ZoomFactorX, const double NumMillisecondsPerWindow, const double NumPixelsPerMillisecond, const double NumMillisecondsPerSample );
+
+protected:
+	void CombineOrSet( FProfilerUIStackNode* ParentUIStackNode, const FProfilerStackNode& ProfilerStackNode, int32 GlobalNodeDepth, const double NumMillisecondsPerWindow, const double NumPixelsPerMillisecond, const double NumMillisecondsPerSample );
+
+
+	/** Converts a tree representation into a flat-array-based. */
+	void LinearizeStream()
+	{
+		SCOPE_LOG_TIME_FUNC();
+
+		LinearRowsOfNodes.SetNum( ThreadNodes.Num()*DEFAULT_VISIBLE_THREAD_DEPTH );
+		for( auto& RowOfNodes : LinearRowsOfNodes )
+		{
+			RowOfNodes.Reset();
+		}
+
+		for( const auto& ThreadNode : ThreadNodes )
+		{
+			LinearizeStream_Recursively( &ThreadNode );
+		}
+	}
+
+	void LinearizeStream_Recursively( const FProfilerUIStackNode* UIStackNode )
+	{
+		if( UIStackNode->GlobalNodeDepth != FProfilerUIStackNode::THREAD_NODE_INDEX )
+		{
+			LinearRowsOfNodes[UIStackNode->GlobalNodeDepth].Add( UIStackNode );
+		}
+
+		for( const auto& UIStackNode : UIStackNode->Children )
+		{
+			LinearizeStream_Recursively( &UIStackNode );
+		}
+	}
 };
 
diff --git a/Engine/Source/Developer/Profiler/Private/Widgets/SProfilerGraphPanel.cpp b/Engine/Source/Developer/Profiler/Private/Widgets/SProfilerGraphPanel.cpp
--- a/Engine/Source/Developer/Profiler/Private/Widgets/SProfilerGraphPanel.cpp
+++ b/Engine/Source/Developer/Profiler/Private/Widgets/SProfilerGraphPanel.cpp
@@ -6,23 +6,36 @@
 #define LOCTEXT_NAMESPACE "SProfilerGraphPanel"
 
 SProfilerGraphPanel::SProfilerGraphPanel()
-	: NumDataPoints( 0 )
+	: bLockMiniViewState( false )
+	, NumDataPoints( 0 )
 	, NumVisiblePoints( 0 )
 	, GraphOffset( 0 )
+	, ViewMode( EProfilerViewMode::InvalidOrMax )
 {
 }
 
 SProfilerGraphPanel::~SProfilerGraphPanel()
 {
 	// Remove ourselves from the profiler manager.
-	if( FProfilerManager::Get().IsValid() )
+	auto ProfilerManager = FProfilerManager::Get();
+	if( ProfilerManager.IsValid() )
 	{
-		FProfilerManager::Get()->OnTrackedStatChanged().RemoveAll( this );
-		FProfilerManager::Get()->OnSessionInstancesUpdated().RemoveAll( this );
-		FProfilerManager::Get()->OnViewModeChanged().RemoveAll( this );
+		ProfilerManager->OnTrackedStatChanged().RemoveAll( this );
+		ProfilerManager->OnSessionInstancesUpdated().RemoveAll( this );
+		ProfilerManager->OnViewModeChanged().RemoveAll( this );
+
+		DataGraph->OnSelectionChangedForIndex().RemoveAll( ProfilerManager.Get() );
+
+		if( ProfilerMiniView.Pin().IsValid() )
+		{
+			ProfilerMiniView.Pin()->OnSelectionBoxChanged().RemoveAll( this );
 	}
 }
 
+	ThreadView->OnViewPositionXChanged().RemoveAll( this );
+	ThreadView->OnViewPositionYChanged().RemoveAll( this );
+}
+
 void SProfilerGraphPanel::Construct( const FArguments& InArgs )
 {
 	ChildSlot
@@ -38,6 +51,9 @@ void SProfilerGraphPanel::Construct( const FArguments& InArgs )
 			[
 				SNew( SVerticalBox )
 
+				// At this moment only one widget of these two can be visible at once.
+				//
+				// DataGraph
 				+ SVerticalBox::Slot()
 				.FillHeight( 1.0f )
 				[
@@ -46,6 +62,15 @@ void SProfilerGraphPanel::Construct( const FArguments& InArgs )
 					.OnViewModeChanged( this, &SProfilerGraphPanel::DataGraph_OnViewModeChanged )
 				]
 
+				// ThreadView
+				+ SVerticalBox::Slot()
+				.FillHeight( 1.0f )
+				[
+					SAssignNew( ThreadView, SProfilerThreadView )
+					//.OnGraphOffsetChanged( this, &SProfilerGraphPanel::OnDataGraphGraphOffsetChanged )
+					//.OnViewModeChanged( this, &SProfilerGraphPanel::DataGraph_OnViewModeChanged )
+				]
+
 				+ SVerticalBox::Slot()
 				.AutoHeight()
 				[
@@ -76,9 +101,12 @@ void SProfilerGraphPanel::Construct( const FArguments& InArgs )
 
 	// Register ourselves with the profiler manager.
 	FProfilerManager::Get()->OnTrackedStatChanged().AddSP( this, &SProfilerGraphPanel::ProfilerManager_OnTrackedStatChanged );
-	FProfilerManager::Get()->OnSessionInstancesUpdated().AddSP( this, &SProfilerGraphPanel::ProfilerManager_OnSessionInstancesUpdated );
-
 	FProfilerManager::Get()->OnViewModeChanged().AddSP( this, &SProfilerGraphPanel::ProfilerManager_OnViewModeChanged );
+
+	DataGraph->OnSelectionChangedForIndex().AddSP( FProfilerManager::Get().ToSharedRef(), &FProfilerManager::DataGraph_OnSelectionChangedForIndex );
+
+	ThreadView->OnViewPositionXChanged().AddSP( this, &SProfilerGraphPanel::ThreadView_OnViewPositionXChanged );
+	ThreadView->OnViewPositionYChanged().AddSP( this, &SProfilerGraphPanel::ThreadView_OnViewPositionYChanged );
 }
 
 
@@ -89,6 +117,8 @@ void SProfilerGraphPanel::Tick( const FGeometry& AllottedGeometry, const double
 
 void SProfilerGraphPanel::HorizontalScrollBar_OnUserScrolled( float ScrollOffset )
 {
+	if( ViewMode == EProfilerViewMode::LineIndexBased )
+	{
 	const float ThumbSizeFraction = FMath::Min( (float)NumVisiblePoints / (float)NumDataPoints, 1.0f );
 	ScrollOffset = FMath::Min( ScrollOffset, 1.0f-ThumbSizeFraction );
  
@@ -97,14 +127,28 @@ void SProfilerGraphPanel::HorizontalScrollBar_OnUserScrolled( float ScrollOffset
 	GraphOffset = FMath::TruncToInt( ScrollOffset * NumDataPoints );
 	DataGraph->ScrollTo( GraphOffset );
 
-	FProfilerManager::Get()->GetProfilerWindow()->ProfilerMiniView->OnSelectionBoxChanged( GraphOffset, GraphOffset + NumVisiblePoints );
+		ProfilerMiniView.Pin()->MoveWithoutZoomSelectionBox( GraphOffset, GraphOffset + NumVisiblePoints );
+	}
+	else if( ViewMode == EProfilerViewMode::ThreadViewTimeBased )
+	{
+		ThreadView->SetPositionXToByScrollBar( ScrollOffset );
+	}
 }
 
 
 void SProfilerGraphPanel::VerticalScrollBar_OnUserScrolled( float ScrollOffset )
 {
+	// @TODO yrx 2014-04-24 Broadcast?
+	if( ViewMode == EProfilerViewMode::LineIndexBased )
+	{
 
 }
+	else if( ViewMode == EProfilerViewMode::ThreadViewTimeBased )
+	{
+		// @TODO yrx 2014-04-23 
+		ThreadView->SetPositonYTo( ScrollOffset );
+	}
+}
 
 
 void SProfilerGraphPanel::ProfilerManager_OnTrackedStatChanged( const FTrackedStat& TrackedStat, bool bIsTracked )
@@ -119,55 +163,109 @@ void SProfilerGraphPanel::ProfilerManager_OnTrackedStatChanged( const FTrackedSt
 	}
 }
 
-void SProfilerGraphPanel::ProfilerManager_OnSessionInstancesUpdated()
-{
-
-}
-
 void SProfilerGraphPanel::OnDataGraphGraphOffsetChanged( int32 InGraphOffset )
 {
 	GraphOffset = InGraphOffset;
 	SetScrollBarState();
-	FProfilerManager::Get()->GetProfilerWindow()->ProfilerMiniView->OnSelectionBoxChanged( InGraphOffset, InGraphOffset + NumVisiblePoints );
+	ProfilerMiniView.Pin()->MoveWithoutZoomSelectionBox( InGraphOffset, InGraphOffset + NumVisiblePoints );
 }
 
-void SProfilerGraphPanel::OnMiniViewSelectionBoxChanged( int32 FrameStart, int32 FrameEnd )
+void SProfilerGraphPanel::MiniView_OnSelectionBoxChanged( int32 FrameStart, int32 FrameEnd )
+{
+	if( ViewMode == EProfilerViewMode::LineIndexBased )
 {
 	GraphOffset = FrameStart;
 	SetScrollBarState();
 	DataGraph->ScrollTo( GraphOffset );
 }
+	else if( ViewMode == EProfilerViewMode::ThreadViewTimeBased )
+	{
+		TGuardValue<bool> LockedMiniViewState( bLockMiniViewState, true );
+		// Update thread-view state.
+		ThreadView->SetFrameRange( FrameStart, FrameEnd );
+	}
+}
+
+void SProfilerGraphPanel::ThreadView_OnViewPositionXChanged( double FrameStartMS, double FrameEndMS, double MaxEndTimeMS, int32 FrameStart, int32 FrameEnd )
+{
+	const double MinStartTimeMS = 0.0f;
+	const double TotalRangeMS = MaxEndTimeMS - MinStartTimeMS;
+	const double ThisFrameRangeMS = FrameEndMS - FrameStartMS;
+
+	const double ThumbSizeFraction = ThisFrameRangeMS / TotalRangeMS;
+	const double OffsetFraction = FrameStartMS / MaxEndTimeMS;
+
+	// Update horizontal scroll bar state.
+	HorizontalScrollBar->SetState( (float)OffsetFraction, (float)ThumbSizeFraction );
+
+	if( !bLockMiniViewState )
+	{
+		// Update profiler mini-view state.
+		ProfilerMiniView.Pin()->MoveAndZoomSelectionBox( FrameStart, FrameEnd );
+	}
+}
+
+void SProfilerGraphPanel::ThreadView_OnViewPositionYChanged( double PosStartY, double PosYEnd, double MaxPosY )
+{
+	const double MinPosY = 0.0f;
+	const double TotalRangeY = MaxPosY - MinPosY;
+	const double ThisRangeY = PosYEnd - PosStartY;
+
+	const double ThumbSizeFraction = ThisRangeY / TotalRangeY;
+	const double OffsetFraction = PosStartY / MaxPosY;
+
+	VerticalScrollBar->SetState( (float)OffsetFraction, (float)ThumbSizeFraction );
+}
 
 void SProfilerGraphPanel::DataGraph_OnViewModeChanged( EDataGraphViewModes::Type ViewMode )
 {
+	if( ViewMode == EProfilerViewMode::LineIndexBased )
+	{
 	UpdateInternals();
 }
+	else if( ViewMode == EProfilerViewMode::ThreadViewTimeBased )
+	{
+		// @TODO yrx 2014-04-23 
+	}
+}
 
 
 void SProfilerGraphPanel::ProfilerManager_OnViewModeChanged( EProfilerViewMode::Type NewViewMode )
 {
 	if( NewViewMode == EProfilerViewMode::LineIndexBased )
 	{
-		HorizontalScrollBar->SetVisibility( EVisibility::Collapsed );
-		HorizontalScrollBar->SetEnabled( false );
+		VerticalScrollBar->SetVisibility( EVisibility::Collapsed );
+		VerticalScrollBar->SetEnabled( false );
+
+		DataGraph->SetVisibility( EVisibility::Visible );
+		DataGraph->SetEnabled( true );
+		ThreadView->SetVisibility( EVisibility::Collapsed );
+		ThreadView->SetEnabled( false );
 	}
 	else if( NewViewMode == EProfilerViewMode::ThreadViewTimeBased )
 	{
-		HorizontalScrollBar->SetVisibility( EVisibility::Visible );
-		HorizontalScrollBar->SetEnabled( true );
+		VerticalScrollBar->SetVisibility( EVisibility::Visible );
+		VerticalScrollBar->SetEnabled( true );
+
+		DataGraph->SetVisibility( EVisibility::Collapsed );
+		DataGraph->SetEnabled( false );
+		ThreadView->SetVisibility( EVisibility::Visible );
+		ThreadView->SetEnabled( true );
 	}
 
-	//DataGraph->SetViewMode( NewViewMode );
+	ViewMode = NewViewMode;
 }
 
 
 void SProfilerGraphPanel::UpdateInternals()
 {
+	if( ViewMode == EProfilerViewMode::LineIndexBased )
+	{
 	NumVisiblePoints = DataGraph->GetNumVisiblePoints();
 	NumDataPoints = DataGraph->GetNumDataPoints();
 
 	SetScrollBarState();
-	FProfilerManager::Get()->GetProfilerWindow()->ProfilerMiniView->OnSelectionBoxChanged( GraphOffset, GraphOffset + NumVisiblePoints );
+		ProfilerMiniView.Pin()->MoveWithoutZoomSelectionBox( GraphOffset, GraphOffset + NumVisiblePoints );
 	 
 	if( FProfilerManager::Get()->IsLivePreview() )
 	{
@@ -175,6 +273,11 @@ void SProfilerGraphPanel::UpdateInternals()
 		HorizontalScrollBar_OnUserScrolled(1.0f);
 	}
 }
+	else if( ViewMode == EProfilerViewMode::ThreadViewTimeBased )
+	{
+		// @TODO yrx 2014-04-23 
+	}
+}
 
 void SProfilerGraphPanel::SetScrollBarState()
 {
@@ -185,4 +288,6 @@ void SProfilerGraphPanel::SetScrollBarState()
 	HorizontalScrollBar->SetState( OffsetFraction, ThumbSizeFraction );
 }
 
+
+
 #undef LOCTEXT_NAMESPACE
diff --git a/Engine/Source/Developer/Profiler/Private/Widgets/SProfilerGraphPanel.h b/Engine/Source/Developer/Profiler/Private/Widgets/SProfilerGraphPanel.h
--- a/Engine/Source/Developer/Profiler/Private/Widgets/SProfilerGraphPanel.h
+++ b/Engine/Source/Developer/Profiler/Private/Widgets/SProfilerGraphPanel.h
@@ -35,9 +35,6 @@ protected:
 	/** Called when the status of specified tracked stat has changed. */
 	void ProfilerManager_OnTrackedStatChanged( const FTrackedStat& TrackedStat, bool bIsTracked );
 
-	/** Called when the list of session instances have changed. */
-	void ProfilerManager_OnSessionInstancesUpdated();
-
 	/**
 	 * Called when the user scrolls the scrollbar
 	 *
@@ -54,7 +51,10 @@ protected:
 
 public:
 	/** Called when the selection box has been changed in the profiler mini-view widget. */
-	void OnMiniViewSelectionBoxChanged( int32 FrameStart, int32 FrameEnd );
+	void MiniView_OnSelectionBoxChanged( int32 FrameStart, int32 FrameEnd );
+
+	void ThreadView_OnViewPositionXChanged( double FrameStartMS, double FrameEndMS, double MaxEndTimeMS, int32 FrameStart, int32 FrameEnd );
+	void ThreadView_OnViewPositionYChanged( double PosYStart, double PosYEnd, double MaxPosY );
 
 protected:
 	/** Sets state of the scroll bar. */
@@ -65,10 +65,20 @@ protected:
 
 	void UpdateInternals();
 
-protected:
+//protected:
+public:
 	/** Holds the data graph widget. */
 	TSharedPtr<SDataGraph> DataGraph;
 
+	/** Holds the thread view widget. */
+	TSharedPtr<SProfilerThreadView> ThreadView;
+
+	/** Weak pointer to the profiler mini-view. */
+	TWeakPtr<SProfilerMiniView> ProfilerMiniView;
+	
+	/** Temporary solution to avoid feedback loop when changing the selection box. */
+	bool bLockMiniViewState;
+
 	/** Horizontal scroll bar, used for scrolling graphs. */
 	TSharedPtr<SScrollBar> HorizontalScrollBar;
 
@@ -83,4 +93,6 @@ protected:
 
 	/** Current offset of the graph, index of the first visible graph point. */
 	int32 GraphOffset;
+
+	EProfilerViewMode::Type ViewMode;
 };
\ No newline at end of file
diff --git a/Engine/Source/Developer/Profiler/Private/Widgets/SProfilerMiniView.cpp b/Engine/Source/Developer/Profiler/Private/Widgets/SProfilerMiniView.cpp
--- a/Engine/Source/Developer/Profiler/Private/Widgets/SProfilerMiniView.cpp
+++ b/Engine/Source/Developer/Profiler/Private/Widgets/SProfilerMiniView.cpp
@@ -32,7 +32,6 @@ void SProfilerMiniView::Reset()
 	FMemory::MemZero( PaintStateMemory );
 	PaintState = nullptr;
 
-	MousePosition = FVector2D::ZeroVector;
 	MousePositionOnButtonDown = FVector2D::ZeroVector;
 
 	SelectionBoxFrameStart = 0;
@@ -119,6 +118,7 @@ int32 SProfilerMiniView::OnPaint( const FGeometry& AllottedGeometry, const FSlat
 	if( IsReady() )
 	{
 		static const FSlateColorBrush SolidWhiteBrush = FSlateColorBrush( FColorList::White );
+		// @TODO yrx 2014-04-24 move to the global scope.
 		const FColor GameThreadColor = FColorList::Red;
 		const FColor RenderThreadColor = FColorList::Blue;
 		const FColor OtherThreadsColor = FColorList::Grey;
@@ -447,7 +447,7 @@ FReply SProfilerMiniView::OnMouseMove( const FGeometry& MyGeometry, const FPoint
 
 	if( IsReady() )
 	{
-		MousePosition = MyGeometry.AbsoluteToLocal( MouseEvent.GetScreenSpacePosition() );
+		const FVector2D MousePosition = MyGeometry.AbsoluteToLocal( MouseEvent.GetScreenSpacePosition() );
 		HoveredFrameIndex = PositionToFrameIndex( MousePosition.X );
 
 		const float CursorPosXDelta = MouseEvent.GetCursorDelta().X;
@@ -467,15 +467,13 @@ FReply SProfilerMiniView::OnMouseMove( const FGeometry& MyGeometry, const FPoint
 					{				
 						SelectionBoxFrameStart = MouseFrameIndex;
 						SelectionBoxFrameStart = FMath::Clamp( SelectionBoxFrameStart, 0, FMath::Max( AllFrames.Num() - SelectionBoxSize - 1, 0 ) );
-						
-						SelectionBoxFrameEnd = SelectionBoxFrameStart + SelectionBoxSize;
+						SelectionBoxFrameEnd = FMath::Min( SelectionBoxFrameStart + SelectionBoxSize, AllFrames.Num() - 1 );
 					}
 					else if( bCanBeEndDragged )
 					{
 						SelectionBoxFrameEnd = MouseFrameIndex;
 						SelectionBoxFrameEnd = FMath::Clamp( SelectionBoxFrameEnd, SelectionBoxSize - 1, AllFrames.Num() - -1 );
-
-						SelectionBoxFrameStart = SelectionBoxFrameEnd - SelectionBoxSize;
+						SelectionBoxFrameStart = FMath::Max( SelectionBoxFrameEnd - SelectionBoxSize, 0 );
 					}
 				}
 				else
@@ -586,11 +584,21 @@ void SProfilerMiniView::AddThreadTime( int32 InFrameIndex, const TMap<uint32, fl
 	StatMetadata = InStatMetaData;
 }
 
-void SProfilerMiniView::OnSelectionBoxChanged( int32 FrameStart, int32 FrameEnd )
+void SProfilerMiniView::MoveWithoutZoomSelectionBox( int32 FrameStart, int32 FrameEnd )
+{
+	const int32 MaxFrameIndex = FMath::Max( AllFrames.Num() - 1, 0 );
+	SelectionBoxFrameStart = FMath::Clamp( FrameStart, 0, MaxFrameIndex );
+	SelectionBoxFrameEnd = FMath::Clamp( FrameEnd, 0, MaxFrameIndex );
+	bAllowSelectionBoxZooming = false;
+}
+
+
+void SProfilerMiniView::MoveAndZoomSelectionBox( int32 FrameStart, int32 FrameEnd )
 {
 	const int32 MaxFrameIndex = FMath::Max( AllFrames.Num() - 1, 0 );
 	SelectionBoxFrameStart = FMath::Clamp( FrameStart, 0, MaxFrameIndex );
 	SelectionBoxFrameEnd = FMath::Clamp( FrameEnd, 0, MaxFrameIndex );
+	//bAllowSelectionBoxZooming = true;
 }
 
 void SProfilerMiniView::MoveSelectionBox( int32 FrameIndex )
diff --git a/Engine/Source/Developer/Profiler/Private/Widgets/SProfilerMiniView.h b/Engine/Source/Developer/Profiler/Private/Widgets/SProfilerMiniView.h
--- a/Engine/Source/Developer/Profiler/Private/Widgets/SProfilerMiniView.h
+++ b/Engine/Source/Developer/Profiler/Private/Widgets/SProfilerMiniView.h
@@ -311,7 +311,8 @@ public:
 	 *	Moves the selection box to the location specified by FrameStart and FrameEnd. 
 	 *	Called by external widgets.
 	 */
-	void OnSelectionBoxChanged( int32 FrameStart, int32 FrameEnd );
+	void MoveWithoutZoomSelectionBox( int32 FrameStart, int32 FrameEnd );
+	void MoveAndZoomSelectionBox( int32 FrameStart, int32 FrameEnd );
 
 	void DrawText( const FString& Text, const FSlateFontInfo& FontInfo, FVector2D Position, const FColor& TextColor, const FColor& ShadowColor, FVector2D ShadowOffset ) const;
 
@@ -436,9 +437,6 @@ protected:
 	uint8 PaintStateMemory[sizeof(FSlateOnPaintState)];
 	mutable FSlateOnPaintState*	PaintState;
 
-	/** The current mouse position. */
-	FVector2D MousePosition;
-
 	/** Mouse position during the call on mouse button down. */
 	FVector2D MousePositionOnButtonDown;
 
diff --git a/Engine/Source/Developer/Profiler/Private/Widgets/SProfilerThreadView.cpp b/Engine/Source/Developer/Profiler/Private/Widgets/SProfilerThreadView.cpp
new file mode 100644
--- /dev/null
+++ b/Engine/Source/Developer/Profiler/Private/Widgets/SProfilerThreadView.cpp
@@ -0,0 +1,665 @@
+// Copyright 1998-2014 Epic Games, Inc. All Rights Reserved.
+
+#include "ProfilerPrivatePCH.h"
+
+#define LOCTEXT_NAMESPACE "SProfilerThreadView"
+
+// SProfilerThreadView
+#undef LOCTEXT_NAMESPACE
+
+SProfilerThreadView::SProfilerThreadView()
+{
+	Reset();
+}
+
+SProfilerThreadView::~SProfilerThreadView()
+{
+
+}
+
+void SProfilerThreadView::Construct( const FArguments& InArgs )
+{
+	BindCommands();
+}
+
+void SProfilerThreadView::Reset()
+{
+	ProfilerStream = nullptr;
+
+	FMemory::MemZero( PaintStateMemory );
+	PaintState = nullptr;
+
+	MousePosition = FVector2D::ZeroVector;
+	LastMousePosition = FVector2D::ZeroVector;
+	MousePositionOnButtonDown = FVector2D::ZeroVector;
+
+	PositionXMS = 0.0f;
+	PositionY = 0.0f;
+	RangeXMS = 0.0f;
+	RangeY = 0.0f;
+	TotalRangeXMS = 0.0f;
+	TotalRangeY = 0.0f;
+	ZoomFactorX = 1.0f;
+
+	NumMillisecondsPerWindow = NUM_MILLISECONDS_PER_WINDOW;
+	NumPixelsPerMillisecond = 0.0f;
+	NumMillisecondsPerSample = 0.0f;
+
+	HoveredFrameIndex = 0;
+	HoveredThreadID = 0;
+	HoveredPositionX = 0.0f;
+	HoveredPositionY = 0.0f;
+
+	DistanceDragged = 0.0f;
+
+	bIsLeftMousePressed = false;
+	bIsRightMousePressed = false;
+	bUpdateData = false;
+	CursorType = EThreadViewCursor::Default;
+}
+
+void SProfilerThreadView::Tick( const FGeometry& AllottedGeometry, const double InCurrentTime, const float InDeltaTime )
+{
+	SCompoundWidget::Tick( AllottedGeometry, InCurrentTime, InDeltaTime );
+
+	if( AllottedGeometry.Size.X > 0.0f )
+	{
+		if( ThisGeometry.Size.X != AllottedGeometry.Size.X )
+		{
+			// Refresh.
+			ThisGeometry = AllottedGeometry;
+			bUpdateData = true;
+		}
+
+		if( ShouldUpdateData() && IsReady() )
+		{
+			UpdateInternalConstants();
+			ProcessData();
+			bUpdateData = false;
+		}
+	}
+}
+
+int32 SProfilerThreadView::OnPaint( const FGeometry& AllottedGeometry, const FSlateRect& MyClippingRect, FSlateWindowElementList& OutDrawElements, int32 LayerId, const FWidgetStyle& InWidgetStyle, bool bParentEnabled ) const
+{
+	//	SCOPE_LOG_TIME_FUNC();
+
+	// Rendering info.
+	const bool bEnabled = ShouldBeEnabled( bParentEnabled );
+	const ESlateDrawEffect::Type DrawEffects = bEnabled ? ESlateDrawEffect::None : ESlateDrawEffect::DisabledEffect;
+	const FSlateBrush* BackgroundBrush = FEditorStyle::GetBrush( "Profiler.LineGraphArea" );
+	const FSlateBrush* WhiteBrush = FEditorStyle::GetBrush( "WhiteTexture" );
+
+	// Paint state for this call to OnPaint, valid only in this scope.
+	PaintState = new((void*)PaintStateMemory) FSlateOnPaintState( AllottedGeometry, MyClippingRect, OutDrawElements, LayerId, InWidgetStyle, DrawEffects );
+
+	// Draw background.
+	FSlateDrawElement::MakeBox
+	(
+		PaintState->OutDrawElements,
+		PaintState->LayerId,
+		PaintState->AllottedGeometry.ToPaintGeometry( FVector2D( 0, 0 ), PaintState->Size2D() ),
+		BackgroundBrush,
+		PaintState->AbsoluteClippingRect,
+		PaintState->DrawEffects,
+		BackgroundBrush->GetTint( InWidgetStyle ) * InWidgetStyle.GetColorAndOpacityTint()
+	);
+	LayerId++;
+
+	// Draw all cycle counters for each thread nodes.
+	if( IsReady() )
+	{
+		DrawFramesBackgroundAndTimelines();
+		DrawUIStackNodes();
+		DrawFrameMarkers();
+	}
+
+#if 1/*DEBUG_PROFILER_PERFORMANCE*/
+
+	LayerId++;
+	// Draw debug information.
+	float GraphDescPosY = PaintState->Size2D().Y - 4.0f * PaintState->SummaryFont8Height;
+
+	// Debug text.
+	FSlateDrawElement::MakeText
+	(
+		OutDrawElements,
+		LayerId,
+		AllottedGeometry.ToOffsetPaintGeometry( FVector2D( 16.0f, GraphDescPosY ) ),
+		FString::Printf( TEXT( "Pos X=%f,Y=%f R X=%f,Y=%f TR X=%f,Y=%f ZF X=%f" ), PositionXMS, PositionY, RangeXMS, RangeY, TotalRangeXMS, TotalRangeY, ZoomFactorX ),
+		PaintState->SummaryFont8,
+		MyClippingRect,
+		DrawEffects,
+		FLinearColor::White
+	);
+	GraphDescPosY -= PaintState->SummaryFont8Height + 1.0f;
+
+	FSlateDrawElement::MakeText
+	(
+		OutDrawElements,
+		LayerId,
+		AllottedGeometry.ToOffsetPaintGeometry( FVector2D( 16.0f, GraphDescPosY ) ),
+		FString::Printf( TEXT( "NumMSPerWin=%f H Fr=%i,TID=%i,PX=%f,PY=%f" ), NumMillisecondsPerWindow, HoveredFrameIndex, HoveredThreadID, HoveredPositionX, HoveredPositionY ),
+		PaintState->SummaryFont8,
+		MyClippingRect,
+		DrawEffects,
+		FLinearColor::White
+	);
+	GraphDescPosY -= PaintState->SummaryFont8Height + 1.0f;
+
+	FSlateDrawElement::MakeText
+	(
+		OutDrawElements,
+		LayerId,
+		AllottedGeometry.ToOffsetPaintGeometry( FVector2D( 16.0f, GraphDescPosY ) ),
+		FString::Printf( TEXT( "DistD=%.2f FI=%3i,%3i" ), DistanceDragged, FramesIndices.X, FramesIndices.Y ),
+		PaintState->SummaryFont8,
+		MyClippingRect,
+		DrawEffects,
+		FLinearColor::White
+	);
+	GraphDescPosY -= PaintState->SummaryFont8Height + 1.0f;
+
+#endif // DEBUG_PROFILER_PERFORMANCE
+
+	// Reset paint state.
+	PaintState = nullptr;
+
+	return SCompoundWidget::OnPaint( AllottedGeometry, MyClippingRect, OutDrawElements, LayerId, InWidgetStyle, bParentEnabled && IsEnabled() );
+}
+
+
+void SProfilerThreadView::DrawFramesBackgroundAndTimelines() const
+{
+	static const FSlateColorBrush SolidWhiteBrush = FSlateColorBrush( FColorList::White );
+
+	check( PaintState );
+	const double ThreadViewOffsetPx = PositionXMS*NumPixelsPerMillisecond;
+	PaintState->LayerId++;
+	TArray<FVector2D> LinePoints;
+
+	// Draw frames background for easier reading.
+	for( const auto& ThreadNode : ProfilerUIStream.ThreadNodes )
+	{
+		if( ThreadNode.StatName == NAME_GameThread )
+		{
+			const FVector2D PositionPx = ThreadNode.GetLocalPosition( ThreadViewOffsetPx, -1.0f );
+			const FVector2D SizePx = FVector2D( ThreadNode.WidthPx, PaintState->Size2D().Y );
+
+			const FSlateRect ClippedFrameBackgroundRect = PaintState->LocalClippingRect.IntersectionWith( FSlateRect( PositionPx, PositionPx + SizePx ) );
+
+			FSlateDrawElement::MakeBox
+			(
+				PaintState->OutDrawElements,
+				PaintState->LayerId,
+				PaintState->AllottedGeometry.ToPaintGeometry( ClippedFrameBackgroundRect.GetPosition(), ClippedFrameBackgroundRect.GetSize() ),
+				&SolidWhiteBrush,
+				PaintState->AbsoluteClippingRect,
+				PaintState->DrawEffects,
+				ThreadNode.FrameIndex % 2 ? FColorList::White.WithAlpha( 64 ) : FColorList::White.WithAlpha( 128 )
+			);
+
+			// Check if this frame time marker is inside the visible area.
+			const float LocalPositionXPx = PositionPx.X + SizePx.X;
+			if( LocalPositionXPx < 0.0f || LocalPositionXPx > PaintState->Size2D().X )
+			{
+				continue;
+			}
+
+			LinePoints.Reset( 2 );
+			LinePoints.Add( FVector2D( LocalPositionXPx, 0.0f ) );
+			LinePoints.Add( FVector2D( LocalPositionXPx, PaintState->Size2D().Y ) );
+
+			// Draw frame time marker.
+			FSlateDrawElement::MakeLines
+			(
+				PaintState->OutDrawElements,
+				PaintState->LayerId,
+				PaintState->AllottedGeometry.ToPaintGeometry(),
+				LinePoints,
+				PaintState->AbsoluteClippingRect,
+				PaintState->DrawEffects,
+				PaintState->WidgetStyle.GetColorAndOpacityTint() * FColorList::SkyBlue,
+				false
+			);
+		}
+	}
+
+	PaintState->LayerId++;
+
+	const double PositionXStartPx = FMath::TruncToFloat( PositionXMS*NumPixelsPerMillisecond / (double)NUM_PIXELS_BETWEEN_TIMELINE )*(double)NUM_PIXELS_BETWEEN_TIMELINE;
+	const double PositionXEndPx = PositionXStartPx + RangeXMS*NumPixelsPerMillisecond;
+
+	for( double TimelinePosXPx = PositionXStartPx; TimelinePosXPx < PositionXEndPx; TimelinePosXPx += (double)NUM_PIXELS_BETWEEN_TIMELINE )
+	{
+		LinePoints.Reset( 2 );
+		LinePoints.Add( FVector2D( TimelinePosXPx - ThreadViewOffsetPx, 0.0f ) );
+		LinePoints.Add( FVector2D( TimelinePosXPx - ThreadViewOffsetPx, PaintState->Size2D().Y ) );
+
+		// Draw time line.
+		FSlateDrawElement::MakeLines
+		(
+			PaintState->OutDrawElements,
+			PaintState->LayerId,
+			PaintState->AllottedGeometry.ToPaintGeometry(),
+			LinePoints,
+			PaintState->AbsoluteClippingRect,
+			PaintState->DrawEffects,
+			PaintState->WidgetStyle.GetColorAndOpacityTint() * FColorList::LimeGreen,
+			false
+		);
+	}
+}
+
+
+void SProfilerThreadView::DrawUIStackNodes() const
+{
+//	SCOPE_LOG_TIME_FUNC();
+	check( PaintState );
+	const double ThreadViewOffsetPx = PositionXMS*NumPixelsPerMillisecond;
+	PaintState->LayerId++;
+
+	static const FSlateBrush* BorderBrush = FEditorStyle::GetBrush( "Profiler.ThreadView.SampleBorder" );
+	const FColor GameThreadColor = FColorList::Red;
+	const FColor RenderThreadColor = FColorList::Blue;
+	const FColor ThreadColors[2] = {GameThreadColor, RenderThreadColor};
+
+	// Draw nodes.
+	for( const auto& RowOfNodes : ProfilerUIStream.LinearRowsOfNodes )
+	{
+		int32 NodeIndex = 0;
+		for( const auto& UIStackNode : RowOfNodes )
+		{
+			NodeIndex++;
+			// Check if the node is visible.
+			//if( UIStackNode->IsVisible() )
+			{
+				const FVector2D PositionPx = UIStackNode->GetLocalPosition( ThreadViewOffsetPx, PositionY ) * FVector2D( 1.0f, NUM_PIXELS_PER_ROW );
+				const FVector2D SizePx = FVector2D( FMath::Max( UIStackNode->WidthPx - 1.0, 0.0 ), NUM_PIXELS_PER_ROW );
+				const FSlateRect ClippedNodeRect = PaintState->LocalClippingRect.IntersectionWith( FSlateRect( PositionPx, PositionPx + SizePx ) );
+
+				// Check if this node is inside the visible area.
+				if( ClippedNodeRect.IsEmpty() )
+				{
+					continue;
+				}
+
+				FColor NodeColor = UIStackNode->bIsCombined ? ThreadColors[UIStackNode->ThreadIndex].WithAlpha( 64 ) : ThreadColors[UIStackNode->ThreadIndex].WithAlpha( 192 );
+				NodeColor.G += NodeIndex % 2 ? 0 : 64;
+
+				// Draw a cycle counter for this profiler UI stack node.
+				FSlateDrawElement::MakeBox
+				(
+					PaintState->OutDrawElements,
+					PaintState->LayerId,
+					PaintState->AllottedGeometry.ToPaintGeometry( ClippedNodeRect.GetPosition(), ClippedNodeRect.GetSize() ),
+					BorderBrush,
+					PaintState->AbsoluteClippingRect,
+					PaintState->DrawEffects,
+					NodeColor
+				);
+			}
+		}
+	}
+
+	// @TODO yrx 2014-04-29 Separate layer for makebox, makeshadowtext, maketext.
+	PaintState->LayerId++;
+
+	const float MarkerPosYOffsetPx = ((float)NUM_PIXELS_PER_ROW - PaintState->SummaryFont8Height)*0.5f;
+	
+	// Draw nodes' descriptions.
+	for( const auto& RowOfNodes : ProfilerUIStream.LinearRowsOfNodes )
+	{
+		for( const auto& UIStackNode : RowOfNodes )
+		{
+			const FVector2D PositionPx = UIStackNode->GetLocalPosition( ThreadViewOffsetPx, PositionY ) * FVector2D( 1.0f, NUM_PIXELS_PER_ROW );
+			const FVector2D SizePx = FVector2D( UIStackNode->WidthPx, NUM_PIXELS_PER_ROW );
+			const FSlateRect ClippedNodeRect = PaintState->LocalClippingRect.IntersectionWith( FSlateRect( PositionPx, PositionPx + SizePx ) );
+
+			// Check if this node is inside the visible area.
+			if( ClippedNodeRect.IsEmpty() )
+			{
+				continue;
+			}
+
+			FString StringStatName = UIStackNode->StatName.GetPlainNameString();
+			FString StringStatNameWithTime = StringStatName + FString::Printf( TEXT( " (%.4f MS)" ), UIStackNode->GetDurationMS() );
+			if( UIStackNode->bIsCulled )
+			{
+				StringStatName += TEXT( " [C]" );
+				StringStatNameWithTime += TEXT( " [C]" );
+			}
+
+			// Update position of the text to be always visible and try to center it.
+			const float StatNameWidthPx = PaintState->FontMeasureService->Measure( StringStatName, PaintState->SummaryFont8 ).X;
+			const float StatNameWithTimeWidthPx = PaintState->FontMeasureService->Measure( StringStatNameWithTime, PaintState->SummaryFont8 ).X;
+			const float TextAreaWidthPx = ClippedNodeRect.GetSize().X;
+
+			bool bUseShortVersion = true;
+			FVector2D AdjustedPositionPx;
+			// Center the stat name with timing if we can.
+			if( TextAreaWidthPx > StatNameWithTimeWidthPx )
+			{
+				AdjustedPositionPx = FVector2D( ClippedNodeRect.Left + (TextAreaWidthPx - StatNameWithTimeWidthPx)*0.5f, PositionPx.Y + MarkerPosYOffsetPx );
+				bUseShortVersion = false;
+			}
+			// Center the stat name.
+			else if( TextAreaWidthPx > StatNameWidthPx )
+			{
+				AdjustedPositionPx = FVector2D( ClippedNodeRect.Left + (TextAreaWidthPx - StatNameWidthPx)*0.5f, PositionPx.Y + MarkerPosYOffsetPx );
+			}
+			// Move to the edge.
+			else
+			{
+				AdjustedPositionPx = FVector2D( ClippedNodeRect.Left, PositionPx.Y + MarkerPosYOffsetPx );
+			}
+
+			const FVector2D AbsolutePositionPx = PaintState->AllottedGeometry.LocalToAbsolute( ClippedNodeRect.GetPosition() );
+			const FSlateRect AbsoluteClippingRect = FSlateRect( AbsolutePositionPx, AbsolutePositionPx + ClippedNodeRect.GetSize() );
+			
+			DrawText( bUseShortVersion ? StringStatName : StringStatNameWithTime, PaintState->SummaryFont8, AdjustedPositionPx, FColorList::White, FColorList::Black, FVector2D( 1.0f, 1.0f ), &AbsoluteClippingRect );
+		}
+	}
+}
+
+
+void SProfilerThreadView::DrawFrameMarkers() const
+{
+	check( PaintState );
+	const double ThreadViewOffsetPx = PositionXMS*NumPixelsPerMillisecond;
+	PaintState->LayerId++;
+
+	for( const auto& ThreadNode : ProfilerUIStream.ThreadNodes )
+	{
+		if( ThreadNode.StatName == NAME_GameThread )
+		{
+			const float MarkerPosXPx = ThreadNode.GetLocalPosition( ThreadViewOffsetPx, 0.0f ).X + ThreadNode.WidthPx;
+
+			// Check if this frame time marker is inside the visible area.
+			if( MarkerPosXPx < 0.0f || MarkerPosXPx > PaintState->Size2D().X )
+			{
+				continue;
+			}
+
+			// Draw text.
+			const FString FrameIndexStr = FString::Printf( TEXT( "%i" ), ThreadNode.FrameIndex );
+			const FString FrameTimesStr = FString::Printf( TEXT( "%.4f [%.4f] MS" ), ThreadNode.CycleCountersEndTimeMS, ThreadNode.GetDurationMS() );
+
+			float MarkerPosYPx = PaintState->Size2D().Y - 2 * PaintState->SummaryFont8Height;
+			DrawText( FrameIndexStr, PaintState->SummaryFont8, FVector2D( MarkerPosXPx, MarkerPosYPx ), FColorList::SkyBlue, FColorList::Black, FVector2D( 1.0f, 1.0f ) );
+
+			MarkerPosYPx += PaintState->SummaryFont8Height;
+			DrawText( FrameTimesStr, PaintState->SummaryFont8, FVector2D( MarkerPosXPx, MarkerPosYPx ), FColorList::SkyBlue, FColorList::Black, FVector2D( 1.0f, 1.0f ) );
+		}
+	}
+
+	PaintState->LayerId++;
+
+	const double PositionXStartPx = FMath::TruncToFloat( PositionXMS*NumPixelsPerMillisecond / (double)NUM_PIXELS_BETWEEN_TIMELINE )*(double)NUM_PIXELS_BETWEEN_TIMELINE;
+	const double PositionXEndPx = PositionXStartPx + RangeXMS*NumPixelsPerMillisecond;
+
+	for( double TimelinePosXPx = PositionXStartPx; TimelinePosXPx < PositionXEndPx; TimelinePosXPx += (double)NUM_PIXELS_BETWEEN_TIMELINE )
+	{
+		const FString TimelineStr = FString::Printf( TEXT( "%.4f MS" ), TimelinePosXPx / NumPixelsPerMillisecond );
+
+		// Draw time line text.
+		float MarkerPosYPx = PaintState->Size2D().Y - 3 * PaintState->SummaryFont8Height;
+		DrawText( TimelineStr, PaintState->SummaryFont8, FVector2D( TimelinePosXPx - ThreadViewOffsetPx, MarkerPosYPx ), FColorList::LimeGreen, FColorList::Black, FVector2D( 1.0f, 1.0f ) );
+	}
+
+#if	0
+	for( int32 FrameIndex = FramesIndices.X; FrameIndex < FramesIndices.Y; ++FrameIndex )
+	{
+		const double FrameTimeMS = ProfilerStream->GetFrameTimeMS( FrameIndex );
+		const double ElapsedTimeMS = ProfilerStream->GetElapsedFrameTimeMS( FrameIndex );
+		const double FrameEndMarkerLocalPosX = ElapsedTimeMS * NumPixelsPerMillisecond - ThreadViewOffsetPx;
+
+		// Draw text.
+		const FVector2D TextLocalPosTop = FVector2D( FrameEndMarkerLocalPosX, PaintState->SummaryFont8Height );
+		const FString FrameTimeStr = FString::Printf( TEXT( "%.4f [%.4f] MS" ), ElapsedTimeMS, FrameTimeMS );
+		DrawText( FrameTimeStr, PaintState->SummaryFont8, TextLocalPosTop, FColorList::SkyBlue, FColorList::Black, FVector2D( 1.0f, 1.0f ) );
+
+		const FVector2D TextLocalPosBottom = FVector2D( FrameEndMarkerLocalPosX, PaintState->Size2D().Y - 2 * PaintState->SummaryFont8Height );
+		DrawText( FrameTimeStr, PaintState->SummaryFont8, TextLocalPosBottom, FColorList::SkyBlue, FColorList::Black, FVector2D( 1.0f, 1.0f ) );
+	}
+#endif // 0
+}
+
+
+void SProfilerThreadView::DrawUIStackNodes_Recursively( const FProfilerUIStackNode& UIStackNode ) const
+{
+	// OBSOLETE
+	check( PaintState );
+	// Don't render thread nodes.
+	if( UIStackNode.ThreadIndex != FProfilerUIStackNode::THREAD_NODE_INDEX )
+	{
+		static const FSlateColorBrush SolidWhiteBrush = FSlateColorBrush( FColorList::White );
+		const FColor GameThreadColor = FColorList::Red;
+	
+		const FVector2D Position = FVector2D( UIStackNode.PositionXPx /*- PositionXMS*/, UIStackNode.PositionY*NUM_PIXELS_PER_ROW );
+		const FVector2D Size = FVector2D( UIStackNode.WidthPx, NUM_PIXELS_PER_ROW );
+
+		// Draw a cycle counter for this profiler UI stack node.
+		FSlateDrawElement::MakeBox
+		(
+			PaintState->OutDrawElements,
+			PaintState->LayerId,
+			PaintState->AllottedGeometry.ToPaintGeometry( Position, Size ),
+			&SolidWhiteBrush,
+			PaintState->AbsoluteClippingRect,
+			PaintState->DrawEffects,
+			GameThreadColor
+		);
+
+		const FString StringStatName = UIStackNode.StatName.GetPlainNameString();
+		DrawText( UIStackNode.StatName.GetPlainNameString(), PaintState->SummaryFont8, Position, FColorList::White, FColorList::Black, FVector2D( 1.0f, 1.0f ) );
+	}
+
+	for( const auto& UIStackNode : UIStackNode.Children )
+	{
+		DrawUIStackNodes_Recursively( UIStackNode );
+	}
+}
+
+
+void SProfilerThreadView::DrawText( const FString& Text, const FSlateFontInfo& FontInfo, FVector2D Position, const FColor& TextColor, const FColor& ShadowColor, FVector2D ShadowOffset, const FSlateRect* ClippingRect /*= nullptr*/ ) const
+{
+	check( PaintState );
+
+	if( ShadowOffset.Size() > 0.0f )
+	{
+		FSlateDrawElement::MakeText
+		(
+			PaintState->OutDrawElements,
+			PaintState->LayerId,
+			PaintState->AllottedGeometry.ToOffsetPaintGeometry( Position + ShadowOffset ),
+			Text,
+			FontInfo,
+			ClippingRect ? *ClippingRect : PaintState->AbsoluteClippingRect,
+			PaintState->DrawEffects,
+			ShadowColor
+		);
+	}
+
+	FSlateDrawElement::MakeText
+	(
+		PaintState->OutDrawElements,
+		++PaintState->LayerId,
+		PaintState->AllottedGeometry.ToOffsetPaintGeometry( Position ),
+		Text,
+		FontInfo,
+		ClippingRect ? *ClippingRect : PaintState->AbsoluteClippingRect,
+		PaintState->DrawEffects,
+		TextColor
+	);
+}
+
+
+FReply SProfilerThreadView::OnMouseButtonDown( const FGeometry& MyGeometry, const FPointerEvent& MouseEvent )
+{
+	FReply Reply = FReply::Unhandled();
+
+	if( IsReady() )
+	{
+		MousePositionOnButtonDown = MyGeometry.AbsoluteToLocal( MouseEvent.GetScreenSpacePosition() );
+
+		if( MouseEvent.GetEffectingButton() == EKeys::LeftMouseButton )
+		{
+			bIsLeftMousePressed = true;
+			DistanceDragged = PositionXMS;
+
+			// Capture mouse, so we can move outside this widget.
+			Reply = FReply::Handled().CaptureMouse( SharedThis( this ) );
+		}
+		else if( MouseEvent.GetEffectingButton() == EKeys::RightMouseButton )
+		{
+			bIsRightMousePressed = true;
+		}
+	}
+
+	return Reply;
+}
+
+FReply SProfilerThreadView::OnMouseButtonUp( const FGeometry& MyGeometry, const FPointerEvent& MouseEvent )
+{
+	FReply Reply = FReply::Unhandled();
+
+	if( IsReady() )
+	{
+		const FVector2D MousePositionOnButtonUp = MyGeometry.AbsoluteToLocal( MouseEvent.GetScreenSpacePosition() );
+		const bool bIsValidForMouseClick = MousePositionOnButtonUp.Equals( MousePositionOnButtonDown, MOUSE_SNAP_DISTANCE );
+
+		if( MouseEvent.GetEffectingButton() == EKeys::LeftMouseButton )
+		{
+			if( bIsLeftMousePressed )
+			{
+				// Release the mouse, we are no longer dragging.
+				Reply = FReply::Handled().ReleaseMouseCapture();
+			}
+
+			bIsLeftMousePressed = false;
+		}
+		else if( MouseEvent.GetEffectingButton() == EKeys::RightMouseButton )
+		{
+			if( bIsRightMousePressed )
+			{
+				if( bIsValidForMouseClick )
+				{
+					ShowContextMenu( MouseEvent.GetScreenSpacePosition() );
+					Reply = FReply::Handled();
+				}
+			}
+			bIsRightMousePressed = false;
+		}
+	}
+
+	return Reply;
+}
+
+FReply SProfilerThreadView::OnMouseMove( const FGeometry& MyGeometry, const FPointerEvent& MouseEvent )
+{
+	//	SCOPE_LOG_TIME_FUNC();
+
+	FReply Reply = FReply::Unhandled();
+
+	if( IsReady() )
+	{
+		const FVector2D MousePosition = MyGeometry.AbsoluteToLocal( MouseEvent.GetScreenSpacePosition() );
+		
+		HoveredPositionX = 0.0;//PositionToFrameIndex( MousePosition.X );
+		HoveredPositionY = 0.0;
+
+		const float CursorPosXDelta = -MouseEvent.GetCursorDelta().X;
+		const float ScrollSpeed = 1.0f / ZoomFactorX;
+
+		if( MouseEvent.IsMouseButtonDown( EKeys::LeftMouseButton ) )
+		{
+			if( HasMouseCapture() && !MouseEvent.GetCursorDelta().IsZero() )
+			{
+				DistanceDragged += CursorPosXDelta*ScrollSpeed*0.1;
+
+				// Inform other widgets that we have scrolled the thread-view.
+				SetPositionX( FMath::Clamp( DistanceDragged, 0.0, TotalRangeXMS - RangeXMS ) );
+				CursorType = EThreadViewCursor::Hand;
+				Reply = FReply::Handled();
+			}
+		}
+		else
+		{
+			CursorType = EThreadViewCursor::Default;
+		}
+	}
+
+	return Reply;
+}
+
+void SProfilerThreadView::OnMouseEnter( const FGeometry& MyGeometry, const FPointerEvent& MouseEvent )
+{
+
+}
+
+void SProfilerThreadView::OnMouseLeave( const FPointerEvent& MouseEvent )
+{
+	if( !HasMouseCapture() )
+	{
+		bIsLeftMousePressed = false;
+		bIsRightMousePressed = false;
+
+		CursorType = EThreadViewCursor::Default;
+	}
+}
+
+FReply SProfilerThreadView::OnMouseWheel( const FGeometry& MyGeometry, const FPointerEvent& MouseEvent )
+{
+	FReply Reply = FReply::Unhandled();
+	
+	const bool bZoomIn = MouseEvent.GetWheelDelta() < 0.0f;
+	const double Center = PositionXMS + RangeXMS*0.5f;
+
+	const double MinVisibleRangeMS = 1.0f / INV_MIN_VISIBLE_RANGE_X;
+	const double NewUnclampedRange = bZoomIn ? RangeXMS*1.25f : RangeXMS / 1.25f;
+	const double NewRange = FMath::Clamp( NewUnclampedRange, MinVisibleRangeMS, FMath::Min( TotalRangeXMS, (double)MAX_VISIBLE_RANGE_X ) );
+
+	const double NewPositionX = FMath::Clamp( Center, NewRange*0.5f, TotalRangeXMS - NewRange*0.5f ) - NewRange*0.5f;
+	SetTimeRange( NewPositionX, NewPositionX + NewRange );
+
+	return Reply;
+}
+
+FReply SProfilerThreadView::OnMouseButtonDoubleClick( const FGeometry& MyGeometry, const FPointerEvent& MouseEvent )
+{
+	FReply Reply = FReply::Unhandled();
+	return Reply;
+}
+
+
+FCursorReply SProfilerThreadView::OnCursorQuery( const FGeometry& MyGeometry, const FPointerEvent& CursorEvent ) const
+{
+	FCursorReply CursorReply = FCursorReply::Unhandled();
+	if( CursorType == EThreadViewCursor::Arrow )
+	{
+		CursorReply = FCursorReply::Cursor( EMouseCursor::ResizeLeftRight );
+	}
+	else if( CursorType == EThreadViewCursor::Hand )
+	{
+		CursorReply = FCursorReply::Cursor( EMouseCursor::GrabHand );
+	}
+	return CursorReply;
+}
+
+void SProfilerThreadView::ShowContextMenu( const FVector2D& ScreenSpacePosition )
+{
+
+}
+
+void SProfilerThreadView::BindCommands()
+{
+
+}
+
+void SProfilerThreadView::ProcessData()
+{
+	//	SCOPE_LOG_TIME_FUNC();
+
+	ProfilerUIStream.GenerateUIStream( *ProfilerStream, PositionXMS, PositionXMS + RangeXMS, ZoomFactorX, NumMillisecondsPerWindow, NumPixelsPerMillisecond, NumMillisecondsPerSample );
+}
diff --git a/Engine/Source/Developer/Profiler/Private/Widgets/SProfilerThreadView.h b/Engine/Source/Developer/Profiler/Private/Widgets/SProfilerThreadView.h
new file mode 100644
--- /dev/null
+++ b/Engine/Source/Developer/Profiler/Private/Widgets/SProfilerThreadView.h
@@ -0,0 +1,442 @@
+// Copyright 1998-2014 Epic Games, Inc. All Rights Reserved.
+
+#pragma once
+
+/*-----------------------------------------------------------------------------
+	Basic structures
+-----------------------------------------------------------------------------*/
+
+
+/*-----------------------------------------------------------------------------
+	Declarations
+-----------------------------------------------------------------------------*/
+
+/** Widget used to present thread data in the mini-view. */
+class SProfilerThreadView : public SCompoundWidget
+{
+	enum 
+	{
+		/** Minimum width of the one rendered sample, if less cycles counter will be combined. */
+		MIN_NUM_PIXELS_PER_SAMPLE = 32,
+
+		/**
+		 *	Number of milliseconds that can be renderer at once in the window.
+		 *	For the default zoom value.
+		 */
+		NUM_MILLISECONDS_PER_WINDOW = 33,
+
+		/** Number of pixels needed to render one row of cycle counter. */
+		NUM_PIXELS_PER_ROW = 16,
+
+		/** Number of pixels. */
+		MOUSE_SNAP_DISTANCE = 4,
+
+		/** Wait time in milliseconds before we display a tooltip. */
+		TOOLTIP_DELAY = 500,
+
+		/** Width of the thread description windows. */
+		WIDTH_THREAD_DESC = 128,
+
+		/**
+		 *	Displayed data will be partitioned into smaller batches to avoid long processing times.
+		 *	This should help in situation when the user scroll the thread-view, so we don't need to wait for the whole data.
+		 *	At the same it adds some overhead to the processing in favor of using massive parallel processing,
+		 *	so overall it will be faster and much more responsive.
+		 *	One partition must have a least one frame.
+		 */
+		NUM_DATA_PARTITIONS = 16,
+		// @TODO yrx 2014-04-25 Dynamic data partitioning
+
+		/**
+		 *	Maximum zoom value for time axis.
+		 *	Default value mean that one 33ms frame can be rendered at once in the window.
+		 *	Maximum zoom allows to see individual cycles.
+		 */
+		INV_MIN_VISIBLE_RANGE_X = 10000,
+		MAX_VISIBLE_RANGE_X = 250,
+
+		/** Number of pixels between each time line. */
+		NUM_PIXELS_BETWEEN_TIMELINE = 96,
+	};
+
+	struct EThreadViewCursor
+	{
+		enum Type
+		{
+			Default,
+			Arrow,
+			Hand,
+		};
+	};
+
+	/** Holds current state provided by OnPaint function, used to simplify drawing. */
+	struct FSlateOnPaintState : public FNoncopyable
+	{
+		FSlateOnPaintState( const FGeometry& InAllottedGeometry, const FSlateRect& InMyClippingRect, FSlateWindowElementList& InOutDrawElements, int32& InLayerId, const FWidgetStyle& InWidgetStyle, ESlateDrawEffect::Type InDrawEffects )
+			: AllottedGeometry( InAllottedGeometry )
+			, AbsoluteClippingRect( InMyClippingRect )
+			, LocalClippingRect( FVector2D::ZeroVector, InAllottedGeometry.Size )
+			, WidgetStyle( InWidgetStyle )
+			, OutDrawElements( InOutDrawElements )
+			, LayerId( InLayerId )
+			, DrawEffects( InDrawEffects )
+			, FontMeasureService( FSlateApplication::Get().GetRenderer()->GetFontMeasureService() )
+			, SummaryFont8( FPaths::EngineContentDir() / TEXT( "Slate/Fonts/Roboto-Regular.ttf" ), 8 )
+			, SummaryFont8Height( FontMeasureService->Measure( TEXT( "!" ), SummaryFont8 ).Y )
+		{}
+
+		const FVector2D& Size2D() const
+		{
+			return AllottedGeometry.Size;
+		}
+
+		/** Accessors. */
+		const FGeometry& AllottedGeometry; 
+		const FSlateRect& AbsoluteClippingRect;
+		const FSlateRect LocalClippingRect;
+		const FWidgetStyle& WidgetStyle;
+		 
+		FSlateWindowElementList& OutDrawElements;
+		int32& LayerId;
+		const ESlateDrawEffect::Type DrawEffects;
+
+		const TSharedRef< FSlateFontMeasure > FontMeasureService;
+
+		const FSlateFontInfo SummaryFont8;
+		const float SummaryFont8Height;
+	};
+
+public:
+	SProfilerThreadView();
+	~SProfilerThreadView();
+
+	SLATE_BEGIN_ARGS( SProfilerThreadView )
+		{}
+	SLATE_END_ARGS()
+
+	/**
+	* Construct this widget
+	*
+	* @param	InArgs	The declaration data for this widget
+	*/
+	void Construct( const FArguments& InArgs );
+
+	/** Resets internal widget's data to the default one. */
+	void Reset();
+
+	/**
+	* Ticks this widget.  Override in derived classes, but always call the parent implementation.
+	*
+	* @param  AllottedGeometry The space allotted for this widget
+	* @param  InCurrentTime  Current absolute real time
+	* @param  InDeltaTime  Real time passed since last tick
+	*/
+	virtual void Tick( const FGeometry& AllottedGeometry, const double InCurrentTime, const float InDeltaTime ) OVERRIDE;
+
+	/**
+	* The widget should respond by populating the OutDrawElements array with FDrawElements
+	* that represent it and any of its children.
+	*
+	* @param AllottedGeometry  The FGeometry that describes an area in which the widget should appear.
+	* @param MyClippingRect    The clipping rectangle allocated for this widget and its children.
+	* @param OutDrawElements   A list of FDrawElements to populate with the output.
+	* @param LayerId           The Layer onto which this widget should be rendered.
+	* @param InColorAndOpacity ColorAverage and Opacity to be applied to all the descendants of the widget being painted
+	* @param bParentEnabled	True if the parent of this widget is enabled.
+	*
+	* @return The maximum layer ID attained by this widget or any of its children.
+	*/
+	virtual int32 OnPaint( const FGeometry& AllottedGeometry, const FSlateRect& MyClippingRect, FSlateWindowElementList& OutDrawElements, int32 LayerId, const FWidgetStyle& InWidgetStyle, bool bParentEnabled ) const OVERRIDE;
+
+	/**
+	* The system calls this method to notify the widget that a mouse button was pressed within it. This event is bubbled.
+	*
+	* @param MyGeometry The Geometry of the widget receiving the event
+	* @param MouseEvent Information about the input event
+	*
+	* @return Whether the event was handled along with possible requests for the system to take action.
+	*/
+	virtual FReply OnMouseButtonDown( const FGeometry& MyGeometry, const FPointerEvent& MouseEvent ) OVERRIDE;
+
+	/**
+	* The system calls this method to notify the widget that a mouse button was release within it. This event is bubbled.
+	*
+	* @param MyGeometry The Geometry of the widget receiving the event
+	* @param MouseEvent Information about the input event
+	*
+	* @return Whether the event was handled along with possible requests for the system to take action.
+	*/
+	virtual FReply OnMouseButtonUp( const FGeometry& MyGeometry, const FPointerEvent& MouseEvent ) OVERRIDE;
+
+	/**
+	* The system calls this method to notify the widget that a mouse moved within it. This event is bubbled.
+	*
+	* @param MyGeometry The Geometry of the widget receiving the event
+	* @param MouseEvent Information about the input event
+	*
+	* @return Whether the event was handled along with possible requests for the system to take action.
+	*/
+	virtual FReply OnMouseMove( const FGeometry& MyGeometry, const FPointerEvent& MouseEvent ) OVERRIDE;
+
+	/**
+	* The system will use this event to notify a widget that the cursor has entered it. This event is NOT bubbled.
+	*
+	* @param MyGeometry The Geometry of the widget receiving the event
+	* @param MouseEvent Information about the input event
+	*/
+	virtual void OnMouseEnter( const FGeometry& MyGeometry, const FPointerEvent& MouseEvent ) OVERRIDE;
+
+	/**
+	* The system will use this event to notify a widget that the cursor has left it. This event is NOT bubbled.
+	*
+	* @param MouseEvent Information about the input event
+	*/
+	virtual void OnMouseLeave( const FPointerEvent& MouseEvent ) OVERRIDE;
+
+	/**
+	* Called when the mouse wheel is spun. This event is bubbled.
+	*
+	* @param  MouseEvent  Mouse event
+	*
+	* @return  Returns whether the event was handled, along with other possible actions
+	*/
+	virtual FReply OnMouseWheel( const FGeometry& MyGeometry, const FPointerEvent& MouseEvent ) OVERRIDE;
+
+	/**
+	* Called when a mouse button is double clicked.  Override this in derived classes.
+	*
+	* @param  InMyGeometry  Widget geometry
+	* @param  InMouseEvent  Mouse button event
+	*
+	* @return  Returns whether the event was handled, along with other possible actions
+	*/
+	virtual FReply OnMouseButtonDoubleClick( const FGeometry& MyGeometry, const FPointerEvent& MouseEvent ) OVERRIDE;
+
+	/**
+	* Called when the system wants to know which cursor to display for this Widget.  This event is bubbled.
+	*
+	* @return  The cursor requested (can be None.)
+	*/
+	virtual FCursorReply OnCursorQuery( const FGeometry& MyGeometry, const FPointerEvent& CursorEvent ) const OVERRIDE;
+
+	void ShowContextMenu( const FVector2D& ScreenSpacePosition );
+
+	/**	Binds UI commands to delegates. */
+	void BindCommands();
+
+	void DrawText( const FString& Text, const FSlateFontInfo& FontInfo, FVector2D Position, const FColor& TextColor, const FColor& ShadowColor, FVector2D ShadowOffset, const FSlateRect* ClippingRect = nullptr ) const;
+
+	void DrawUIStackNodes_Recursively( const FProfilerUIStackNode& UIStackNode ) const;
+
+	void DrawFramesBackgroundAndTimelines() const;
+	void DrawUIStackNodes() const;
+	void DrawFrameMarkers() const;
+
+public:
+
+	/**
+	 *	Changes the position-x of the thread view.
+	 *	Called by the horizontal scroll bar.
+	 */
+	void SetPositionXToByScrollBar( double ScrollOffset ) 
+	{
+		SetPositionX( ScrollOffset*TotalRangeXMS );
+	}
+
+	void SetPositionX( double NewPositionXMS )
+	{
+		const double ClampedPositionXMS = FMath::Clamp( NewPositionXMS, 0.0, TotalRangeXMS - RangeXMS );
+		SetTimeRange( ClampedPositionXMS, ClampedPositionXMS + RangeXMS, true );
+	}
+
+	/**
+	 *	Changes the position-y of the thread view.
+	 *	Called by the external code.
+	 */
+	void SetPositonYTo( double ScrollOffset )
+	{
+
+	}
+
+	/** Changes the position-x and range-x of the thread view. */
+	void SetTimeRange( double StartTimeMS, double EndTimeMS, bool bBroadcast = true )
+	{
+		check( EndTimeMS > StartTimeMS );
+
+		PositionXMS = StartTimeMS;
+		RangeXMS = EndTimeMS - StartTimeMS;
+		FramesIndices = ProfilerStream->GetFramesIndicesForTimeRange( StartTimeMS, EndTimeMS );
+
+		bUpdateData = true;
+		
+		UE_LOG( LogTemp, Log, TEXT( "StartTimeMS=%f, EndTimeMS=%f, bBroadcast=%1i FramesIndices=%3i,%3i" ), StartTimeMS, EndTimeMS, (int)bBroadcast, FramesIndices.X, FramesIndices.Y );
+		if( bBroadcast )
+		{
+			ViewPositionXChangedEvent.Broadcast( StartTimeMS, EndTimeMS, TotalRangeXMS, FramesIndices.X, FramesIndices.Y );
+		}
+	}
+
+	/**
+	 *	Changes the position-x and range-x of the thread view. 
+	 *	Called by the mini-view.
+	 */
+	void SetFrameRange( int32 FrameStart, int32 FrameEnd )
+	{
+		const double EndTimeMS = ProfilerStream->GetElapsedFrameTimeMS( FrameEnd );
+		const double StartTimeMS = ProfilerStream->GetElapsedFrameTimeMS( FrameStart ) - ProfilerStream->GetFrameTimeMS( FrameStart );
+		SetTimeRange( StartTimeMS, EndTimeMS, true );
+	}
+
+	/** Attaches profiler stream to the thread-view widgets and displays the first frame of data. */
+	void AttachProfilerStream( const FProfilerStream& InProfilerStream )
+	{
+		ProfilerStream = &InProfilerStream;
+
+		TotalRangeXMS = ProfilerStream->GetElapsedTime();
+		TotalRangeY = ProfilerStream->GetNumThreads()*FProfilerUIStream::DEFAULT_VISIBLE_THREAD_DEPTH;
+
+		// Display the first frame.
+		const FProfilerFrame* ProfilerFrame = ProfilerStream->GetProfilerFrame( 0 );
+		SetTimeRange( ProfilerFrame->Root->CycleCounterStartTimeMS, ProfilerFrame->Root->CycleCounterEndTimeMS );
+	}
+
+public:
+	/** The event to execute when the position-x of the thread view has been changed. */
+	DECLARE_EVENT_FiveParams( SProfilerThreadView, FViewPositionXChangedEvent, double /*StartTimeMS*/, double /*EndTimeMS*/, double /*MaxEndTimeMS*/, int32 /*FrameStart*/, int32 /*FrameEnd*/ );
+	FViewPositionXChangedEvent& OnViewPositionXChanged()
+	{
+		return ViewPositionXChangedEvent;
+	}
+	
+protected:
+	/** The event to execute when the position-x of the thread view has been changed. */
+	FViewPositionXChangedEvent ViewPositionXChangedEvent;
+
+public:
+	/** The event to execute when the position-y of the thread view has been changed. */
+	DECLARE_EVENT_ThreeParams( SProfilerThreadView, FViewPositionYChangedEvent, double /*PosYStart*/, double /*PosYEnd*/, double /*MaxPosY*/ );
+	FViewPositionYChangedEvent& OnViewPositionYChanged()
+	{
+		return ViewPositionYChangedEvent;
+	}
+
+protected:
+	/** The event to execute when the position-y of the thread view has been changed. */
+	FViewPositionYChangedEvent ViewPositionYChangedEvent;
+
+
+protected:
+	void UpdateInternalConstants()
+	{
+		ZoomFactorX = (double)NUM_MILLISECONDS_PER_WINDOW / RangeXMS;
+		RangeY = FMath::RoundToFloat( ThisGeometry.Size.Y / (double)NUM_PIXELS_PER_ROW );
+		
+		const double Aspect = ThisGeometry.Size.X / NUM_MILLISECONDS_PER_WINDOW * ZoomFactorX;
+		NumMillisecondsPerWindow = (double)ThisGeometry.Size.X / Aspect;
+		NumPixelsPerMillisecond = (double)ThisGeometry.Size.X / NumMillisecondsPerWindow;
+		NumMillisecondsPerSample = NumMillisecondsPerWindow / (double)ThisGeometry.Size.X * (double)MIN_NUM_PIXELS_PER_SAMPLE;	
+	}
+	
+	void ProcessData();
+
+	/**
+	 * @return True, if the widget is ready to use, also means that contains at least one frame of the thread data.
+	 */
+	bool IsReady() const
+	{
+		return ProfilerStream && ProfilerStream->GetNumFrames() > 0;
+	}
+
+	bool ShouldUpdateData()
+	{
+		return bUpdateData;
+	}
+
+protected:
+
+	/*-----------------------------------------------------------------------------
+		Data variables
+	-----------------------------------------------------------------------------*/
+
+	/** Profiler UI stream, contains data optimized for displaying in this widget. */
+	FProfilerUIStream ProfilerUIStream;
+
+	/** Pointer to the profiler stream, used as a source for the UI stream. */
+	const FProfilerStream* ProfilerStream;
+
+	/*-----------------------------------------------------------------------------
+		UI variables
+	-----------------------------------------------------------------------------*/
+
+	FGeometry ThisGeometry;
+
+	/** Current Slate OnPaint state. */
+	uint8 PaintStateMemory[sizeof(FSlateOnPaintState)];
+	mutable FSlateOnPaintState*	PaintState;
+
+	/** The current mouse position. */
+	FVector2D MousePosition;
+
+	/** The last mouse position. */
+	FVector2D LastMousePosition;
+
+	/** Mouse position during the call on mouse button down. */
+	FVector2D MousePositionOnButtonDown;
+
+	/** Position-X of the thread view, in milliseconds. */
+	double PositionXMS;
+
+	/** Position-Y of the thread view, where 1.0 means one row of the data. */
+	double PositionY;
+
+	/** Range of the visible data for the current zoom, in milliseconds. */
+	double RangeXMS;
+
+	/** Range of the visible data. */
+	double RangeY;
+
+	/** Range of the all collected data, in milliseconds. */
+	double TotalRangeXMS;
+
+	/** Range of the all collected data. */
+	double TotalRangeY;
+
+	/** Current zoom value for X. */
+	double ZoomFactorX;
+
+	/** Number of milliseconds that can be renderer at once in the window. */
+	double NumMillisecondsPerWindow;
+
+	/** Number of pixels needed to render one millisecond cycle counter. */
+	double NumPixelsPerMillisecond;
+
+	/** Number of milliseconds that can be displayed as one cycle counter. */
+	double NumMillisecondsPerSample;
+
+	/** Index of the frame currently being hovered by the mouse. */
+	int32 HoveredFrameIndex;
+
+	/** Thread ID currently being hovered by the mouse. */
+	int32 HoveredThreadID;
+
+	/** Position-Y of the thread view currently being hovered by the mouse, in milliseconds. */
+	double HoveredPositionX;
+
+	/** Position-Y of the thread view currently being hovered by the mouse. */
+	double HoveredPositionY;
+
+	/** Distance dragged. */
+	double DistanceDragged;
+
+	/** Frame indices of the currently visible data. X= FrameStart, Y=FrameEnd+1 */
+	FIntPoint FramesIndices;
+
+	bool bIsLeftMousePressed;
+	bool bIsRightMousePressed;
+
+	/** Whether to updated data. */
+	bool bUpdateData;
+
+	/** Cursor type. */
+	EThreadViewCursor::Type CursorType;
+};
\ No newline at end of file
diff --git a/Engine/Source/Developer/Profiler/Private/Widgets/SProfilerWindow.cpp b/Engine/Source/Developer/Profiler/Private/Widgets/SProfilerWindow.cpp
--- a/Engine/Source/Developer/Profiler/Private/Widgets/SProfilerWindow.cpp
+++ b/Engine/Source/Developer/Profiler/Private/Widgets/SProfilerWindow.cpp
@@ -78,7 +78,13 @@ SProfilerWindow::SProfilerWindow()
 {}
 
 SProfilerWindow::~SProfilerWindow()
-{}
+{
+	// Remove ourselves from the profiler manager.
+	if( FProfilerManager::Get().IsValid() )
+	{
+		FProfilerManager::Get()->OnViewModeChanged().RemoveAll( this );
+	}
+}
 
 BEGIN_SLATE_FUNCTION_BUILD_OPTIMIZATION
 void SProfilerWindow::Construct( const FArguments& InArgs )
@@ -147,11 +153,12 @@ void SProfilerWindow::Construct( const FArguments& InArgs )
 				.FillHeight( 1.0f )
 				.Padding(0.0f, 8.0f, 0.0f, 0.0f)
 				[
-					SNew(SHorizontalBox)
+					SNew( SHorizontalBox )
 					.IsEnabled( this, &SProfilerWindow::IsProfilerEnabled )
 
 					+SHorizontalBox::Slot()
 					.AutoWidth()
+					.Expose( FiltersAndPresetsSlot )
 					[
 						SNew(SBox)
 						.WidthOverride( 256.0f )
@@ -191,7 +198,7 @@ void SProfilerWindow::Construct( const FArguments& InArgs )
 							+SVerticalBox::Slot()
 							.FillHeight( 1.0f )
 							[
-								SNew(SFiltersAndPresets)
+								SAssignNew( FiltersAndPresets, SFiltersAndPresets )
 							]
 						]
 					]
@@ -287,9 +294,9 @@ void SProfilerWindow::Construct( const FArguments& InArgs )
 			.Expose( OverlaySettingsSlot )
 		];
 
-	GraphPanel->GetMainDataGraph()->OnSelectionChangedForIndex().AddSP( FProfilerManager::Get().ToSharedRef(), &FProfilerManager::DataGraph_OnSelectionChangedForIndex );
-
-	ProfilerMiniView->OnSelectionBoxChanged().AddSP( GraphPanel.ToSharedRef(), &SProfilerGraphPanel::OnMiniViewSelectionBoxChanged );
+	ProfilerMiniView->OnSelectionBoxChanged().AddSP( GraphPanel.ToSharedRef(), &SProfilerGraphPanel::MiniView_OnSelectionBoxChanged );
+	FProfilerManager::Get()->OnViewModeChanged().AddSP( this, &SProfilerWindow::ProfilerManager_OnViewModeChanged );
+	GraphPanel->ProfilerMiniView = ProfilerMiniView;
 }
 
 void SProfilerWindow::ManageEventGraphTab( const FGuid ProfilerInstanceID, const bool bCreateFakeTab, const FString TabName )
@@ -558,5 +565,29 @@ void SProfilerWindow::CloseProfilerSettings()
 	MainContentPanel->SetEnabled( true );
 }
 
+void SProfilerWindow::ProfilerManager_OnViewModeChanged( EProfilerViewMode::Type NewViewMode )
+{
+	if( NewViewMode == EProfilerViewMode::LineIndexBased )
+	{
+		EventGraphPanel->SetVisibility( EVisibility::Visible );
+		EventGraphPanel->SetEnabled( true );
+
+		(*FiltersAndPresetsSlot)
+		[
+			FiltersAndPresets.ToSharedRef()
+		];
+	}
+	else if( NewViewMode == EProfilerViewMode::ThreadViewTimeBased )
+	{
+		EventGraphPanel->SetVisibility( EVisibility::Collapsed );
+		EventGraphPanel->SetEnabled( false );
+
+		(*FiltersAndPresetsSlot)
+		[
+			SNullWidget::NullWidget
+		];
+	}
+}
+
 
 #undef LOCTEXT_NAMESPACE
diff --git a/Engine/Source/Developer/Profiler/Private/Widgets/SProfilerWindow.h b/Engine/Source/Developer/Profiler/Private/Widgets/SProfilerWindow.h
--- a/Engine/Source/Developer/Profiler/Private/Widgets/SProfilerWindow.h
+++ b/Engine/Source/Developer/Profiler/Private/Widgets/SProfilerWindow.h
@@ -58,6 +58,8 @@ protected:
 
 	void SendingServiceSideCapture_Load( const FString Filename );
 
+	void ProfilerManager_OnViewModeChanged( EProfilerViewMode::Type NewViewMode );
+
 private:
 	/**
 	 * Called after a key is pressed when this widget has keyboard focus
@@ -92,9 +94,14 @@ private:
 	/** Holds all widgets for the profiler window like menu bar, toolbar and tabs. */
 	TSharedPtr<SVerticalBox> MainContentPanel;
 
+public:
 	/** Holds all event graphs. */
 	TSharedPtr<SVerticalBox> EventGraphPanel;
 
+	/** Holds the filter and presets widget/slot. */
+	SHorizontalBox::FSlot* FiltersAndPresetsSlot;
+	TSharedPtr<SFiltersAndPresets> FiltersAndPresets;
+
 	/** Widget for the panel which contains all graphs and event graphs. */
 	TSharedPtr<SProfilerGraphPanel> GraphPanel;
 
@@ -110,6 +117,6 @@ private:
 	/** Active event graphs, one event graph for each profiler instance, stored as FGuid -> SEventGraph. */
 	TMap< FGuid, TSharedRef<class SEventGraph> > ActiveEventGraphs;
 
-public:
+
 	TSharedPtr<SProfilerMiniView> ProfilerMiniView;
 };
\ No newline at end of file
diff --git a/Engine/Source/Runtime/Core/Private/Stats/StatsFile.cpp b/Engine/Source/Runtime/Core/Private/Stats/StatsFile.cpp
--- a/Engine/Source/Runtime/Core/Private/Stats/StatsFile.cpp
+++ b/Engine/Source/Runtime/Core/Private/Stats/StatsFile.cpp
@@ -68,13 +68,13 @@ FStatsThreadState::FStatsThreadState( FString const& Filename )
 		// Verify frames offsets.
 		for( int32 FrameIndex = 0; FrameIndex < Stream.FramesInfo.Num(); ++FrameIndex )
 		{
-			const int64 FrameOffset = Stream.FramesInfo[FrameIndex].FrameOffset;
-			FileReader->Seek( FrameOffset );
+			const int64 FrameFileOffset = Stream.FramesInfo[FrameIndex].FrameFileOffset;
+			FileReader->Seek( FrameFileOffset );
 
 			int64 TargetFrame;
 			*FileReader << TargetFrame;
 		}
-		FileReader->Seek( Stream.FramesInfo[0].FrameOffset );
+		FileReader->Seek( Stream.FramesInfo[0].FrameFileOffset );
 
 		// Read the raw stats messages.
 		FStatPacketArray IncomingData;
diff --git a/Engine/Source/Runtime/Core/Public/Stats/StatsFile.h b/Engine/Source/Runtime/Core/Public/Stats/StatsFile.h
--- a/Engine/Source/Runtime/Core/Public/Stats/StatsFile.h
+++ b/Engine/Source/Runtime/Core/Public/Stats/StatsFile.h
@@ -123,24 +123,24 @@ struct FStatsFrameInfo
 {
 	/** Empty constructor. */
 	FStatsFrameInfo()
-		: FrameOffset(0)
+		: FrameFileOffset(0)
 	{}
 
 	/** Initialization constructor. */
-	FStatsFrameInfo( int64 InFrameOffset, TMap<uint32, int64>& InThreadCycles )
-		: FrameOffset(InFrameOffset)
+	FStatsFrameInfo( int64 InFrameFileOffset, TMap<uint32, int64>& InThreadCycles )
+		: FrameFileOffset(InFrameFileOffset)
 		, ThreadCycles(InThreadCycles)
 	{}
 
 	/** Serialization operator. */
 	friend FArchive& operator << (FArchive& Ar, FStatsFrameInfo& Data)
 	{
-		Ar << Data.ThreadCycles << Data.FrameOffset;
+		Ar << Data.ThreadCycles << Data.FrameFileOffset;
 		return Ar;
 	}
 
 	/** Frame offset in the stats file. */
-	int64 FrameOffset;
+	int64 FrameFileOffset;
 
 	/** Thread cycles for this frame. */
 	TMap<uint32, int64> ThreadCycles;
