Profiler - Disabled debug code

diff --git a/Engine/Source/Developer/Profiler/Private/ProfilerStream.cpp b/Engine/Source/Developer/Profiler/Private/ProfilerStream.cpp
--- a/Engine/Source/Developer/Profiler/Private/ProfilerStream.cpp
+++ b/Engine/Source/Developer/Profiler/Private/ProfilerStream.cpp
@@ -8,7 +8,7 @@
 
 void FProfilerUIStream::GenerateUIStream( const FProfilerStream& ProfilerStream, const double StartTimeMS, const double EndTimeMS, const double ZoomFactorX, const double NumMillisecondsPerWindow, const double NumPixelsPerMillisecond, const double NumMillisecondsPerSample )
 {
-	SCOPE_LOG_TIME_FUNC();
+	//SCOPE_LOG_TIME_FUNC();
 
 	ProfilerStream;
 	ZoomFactorX;
diff --git a/Engine/Source/Developer/Profiler/Private/ProfilerStream.h b/Engine/Source/Developer/Profiler/Private/ProfilerStream.h
--- a/Engine/Source/Developer/Profiler/Private/ProfilerStream.h
+++ b/Engine/Source/Developer/Profiler/Private/ProfilerStream.h
@@ -511,7 +511,7 @@ protected:
 	/** Converts a tree representation into a flat-array-based. */
 	void LinearizeStream()
 	{
-		SCOPE_LOG_TIME_FUNC();
+		//SCOPE_LOG_TIME_FUNC();
 
 		LinearRowsOfNodes.SetNum( ThreadNodes.Num()*DEFAULT_VISIBLE_THREAD_DEPTH );
 		for( auto& RowOfNodes : LinearRowsOfNodes )
diff --git a/Engine/Source/Developer/Profiler/Private/Widgets/SEventGraph.cpp b/Engine/Source/Developer/Profiler/Private/Widgets/SEventGraph.cpp
--- a/Engine/Source/Developer/Profiler/Private/Widgets/SEventGraph.cpp
+++ b/Engine/Source/Developer/Profiler/Private/Widgets/SEventGraph.cpp
@@ -3361,16 +3361,16 @@ bool SEventGraph::SelectAllFrames_CanExecute() const
 
 void SEventGraph::ProfilerManager_OnViewModeChanged( EProfilerViewMode::Type NewViewMode )
 {
-	if( NewViewMode == EProfilerViewMode::LineIndexBased )
-	{
-		FunctionDetailsBox->SetVisibility( EVisibility::Visible );
-		FunctionDetailsBox->SetEnabled( true );
-	}
-	else if( NewViewMode == EProfilerViewMode::ThreadViewTimeBased )
-	{
-		FunctionDetailsBox->SetVisibility( EVisibility::Collapsed );
-		FunctionDetailsBox->SetEnabled( false );
-	}
+// 	if( NewViewMode == EProfilerViewMode::LineIndexBased )
+// 	{
+// 		FunctionDetailsBox->SetVisibility( EVisibility::Visible );
+// 		FunctionDetailsBox->SetEnabled( true );
+// 	}
+// 	else if( NewViewMode == EProfilerViewMode::ThreadViewTimeBased )
+// 	{
+// 		FunctionDetailsBox->SetVisibility( EVisibility::Collapsed );
+// 		FunctionDetailsBox->SetEnabled( false );
+// 	}
 }
 
 #undef LOCTEXT_NAMESPACE
diff --git a/Engine/Source/Developer/Profiler/Private/Widgets/SProfilerThreadView.h b/Engine/Source/Developer/Profiler/Private/Widgets/SProfilerThreadView.h
--- a/Engine/Source/Developer/Profiler/Private/Widgets/SProfilerThreadView.h
+++ b/Engine/Source/Developer/Profiler/Private/Widgets/SProfilerThreadView.h
@@ -269,7 +269,7 @@ public:
 
 		bUpdateData = true;
 		
-		UE_LOG( LogTemp, Log, TEXT( "StartTimeMS=%f, EndTimeMS=%f, bBroadcast=%1i FramesIndices=%3i,%3i" ), StartTimeMS, EndTimeMS, (int)bBroadcast, FramesIndices.X, FramesIndices.Y );
+		//UE_LOG( LogTemp, Log, TEXT( "StartTimeMS=%f, EndTimeMS=%f, bBroadcast=%1i FramesIndices=%3i,%3i" ), StartTimeMS, EndTimeMS, (int)bBroadcast, FramesIndices.X, FramesIndices.Y );
 		if( bBroadcast )
 		{
 			ViewPositionXChangedEvent.Broadcast( StartTimeMS, EndTimeMS, TotalRangeXMS, FramesIndices.X, FramesIndices.Y );
