Fixed Clang compile errors in SProfilerGraphPanel.cpp

diff --git a/Engine/Source/Developer/Profiler/Private/Widgets/SProfilerGraphPanel.cpp b/Engine/Source/Developer/Profiler/Private/Widgets/SProfilerGraphPanel.cpp
--- a/Engine/Source/Developer/Profiler/Private/Widgets/SProfilerGraphPanel.cpp
+++ b/Engine/Source/Developer/Profiler/Private/Widgets/SProfilerGraphPanel.cpp
@@ -219,11 +219,11 @@ void SProfilerGraphPanel::ThreadView_OnViewPositionYChanged( double PosStartY, d
 
 void SProfilerGraphPanel::DataGraph_OnViewModeChanged( EDataGraphViewModes::Type ViewMode )
 {
-	if( ViewMode == EProfilerViewMode::LineIndexBased )
+	if( ViewMode == EDataGraphViewModes::Index )
 	{
-	UpdateInternals();
-}
-	else if( ViewMode == EProfilerViewMode::ThreadViewTimeBased )
+		UpdateInternals();
+	}
+	else if( ViewMode == EDataGraphViewModes::Time )
 	{
 		// @TODO yrx 2014-04-23 
 	}
