Profiler - Removed unused functionality, fixed bad formatting

diff --git a/Engine/Source/Developer/Profiler/Private/Widgets/SProfilerGraphPanel.cpp b/Engine/Source/Developer/Profiler/Private/Widgets/SProfilerGraphPanel.cpp
--- a/Engine/Source/Developer/Profiler/Private/Widgets/SProfilerGraphPanel.cpp
+++ b/Engine/Source/Developer/Profiler/Private/Widgets/SProfilerGraphPanel.cpp
@@ -11,8 +11,7 @@ SProfilerGraphPanel::SProfilerGraphPanel()
 	, NumVisiblePoints( 0 )
 	, GraphOffset( 0 )
 	, ViewMode( EProfilerViewMode::InvalidOrMax )
-{
-}
+{}
 
 SProfilerGraphPanel::~SProfilerGraphPanel()
 {
@@ -29,8 +28,8 @@ SProfilerGraphPanel::~SProfilerGraphPanel()
 		if( ProfilerMiniView.Pin().IsValid() )
 		{
 			ProfilerMiniView.Pin()->OnSelectionBoxChanged().RemoveAll( this );
+		}
 	}
-}
 
 	ThreadView->OnViewPositionXChanged().RemoveAll( this );
 	ThreadView->OnViewPositionYChanged().RemoveAll( this );
@@ -41,12 +40,12 @@ void SProfilerGraphPanel::Construct( const FArguments& InArgs )
 	ChildSlot
 	[
 		SNew( SBorder )
-		.BorderImage( FEditorStyle::GetBrush("ToolPanel.GroupBorder") )
+		.BorderImage( FEditorStyle::GetBrush( "ToolPanel.GroupBorder" ) )
 		.Padding( 2.0f )
 		[
 			SNew( SHorizontalBox )
 
-			+SHorizontalBox::Slot()
+			+ SHorizontalBox::Slot()
 			.FillWidth( 1.0f )
 			[
 				SNew( SVerticalBox )
@@ -54,12 +53,11 @@ void SProfilerGraphPanel::Construct( const FArguments& InArgs )
 				// At this moment only one widget of these two can be visible at once.
 				//
 				// DataGraph
-				+ SVerticalBox::Slot()
+				+SVerticalBox::Slot()
 				.FillHeight( 1.0f )
 				[
 					SAssignNew( DataGraph, SDataGraph )
 					.OnGraphOffsetChanged( this, &SProfilerGraphPanel::OnDataGraphGraphOffsetChanged )
-					.OnViewModeChanged( this, &SProfilerGraphPanel::DataGraph_OnViewModeChanged )
 				]
 
 				// ThreadView
@@ -67,8 +65,6 @@ void SProfilerGraphPanel::Construct( const FArguments& InArgs )
 				.FillHeight( 1.0f )
 				[
 					SAssignNew( ThreadView, SProfilerThreadView )
-					//.OnGraphOffsetChanged( this, &SProfilerGraphPanel::OnDataGraphGraphOffsetChanged )
-					//.OnViewModeChanged( this, &SProfilerGraphPanel::DataGraph_OnViewModeChanged )
 				]
 
 				+ SVerticalBox::Slot()
@@ -83,7 +79,7 @@ void SProfilerGraphPanel::Construct( const FArguments& InArgs )
 				]
 			]
 
-			+SHorizontalBox::Slot()
+			+ SHorizontalBox::Slot()
 			.AutoWidth()
 			[
 				SAssignNew( VerticalScrollBar, SScrollBar )
@@ -93,7 +89,7 @@ void SProfilerGraphPanel::Construct( const FArguments& InArgs )
 				.Thickness( FVector2D( 8.0f, 8.0f ) )
 				.OnUserScrolled( this, &SProfilerGraphPanel::VerticalScrollBar_OnUserScrolled )
 			]
-		]	
+		]
 	];
 
 	HorizontalScrollBar->SetState( 0.0f, 1.0f );
@@ -119,13 +115,13 @@ void SProfilerGraphPanel::HorizontalScrollBar_OnUserScrolled( float ScrollOffset
 {
 	if( ViewMode == EProfilerViewMode::LineIndexBased )
 	{
-	const float ThumbSizeFraction = FMath::Min( (float)NumVisiblePoints / (float)NumDataPoints, 1.0f );
-	ScrollOffset = FMath::Min( ScrollOffset, 1.0f-ThumbSizeFraction );
- 
-	HorizontalScrollBar->SetState( ScrollOffset, ThumbSizeFraction );
- 
-	GraphOffset = FMath::TruncToInt( ScrollOffset * NumDataPoints );
-	DataGraph->ScrollTo( GraphOffset );
+		const float ThumbSizeFraction = FMath::Min( (float)NumVisiblePoints / (float)NumDataPoints, 1.0f );
+		ScrollOffset = FMath::Min( ScrollOffset, 1.0f - ThumbSizeFraction );
+
+		HorizontalScrollBar->SetState( ScrollOffset, ThumbSizeFraction );
+
+		GraphOffset = FMath::TruncToInt( ScrollOffset * NumDataPoints );
+		DataGraph->ScrollTo( GraphOffset );
 
 		ProfilerMiniView.Pin()->MoveWithoutZoomSelectionBox( GraphOffset, GraphOffset + NumVisiblePoints );
 	}
@@ -142,7 +138,7 @@ void SProfilerGraphPanel::VerticalScrollBar_OnUserScrolled( float ScrollOffset )
 	if( ViewMode == EProfilerViewMode::LineIndexBased )
 	{
 
-}
+	}
 	else if( ViewMode == EProfilerViewMode::ThreadViewTimeBased )
 	{
 		// @TODO yrx 2014-04-23 
@@ -173,11 +169,11 @@ void SProfilerGraphPanel::OnDataGraphGraphOffsetChanged( int32 InGraphOffset )
 void SProfilerGraphPanel::MiniView_OnSelectionBoxChanged( int32 FrameStart, int32 FrameEnd )
 {
 	if( ViewMode == EProfilerViewMode::LineIndexBased )
-{
-	GraphOffset = FrameStart;
-	SetScrollBarState();
-	DataGraph->ScrollTo( GraphOffset );
-}
+	{
+		GraphOffset = FrameStart;
+		SetScrollBarState();
+		DataGraph->ScrollTo( GraphOffset );
+	}
 	else if( ViewMode == EProfilerViewMode::ThreadViewTimeBased )
 	{
 		TGuardValue<bool> LockedMiniViewState( bLockMiniViewState, true );
@@ -217,19 +213,6 @@ void SProfilerGraphPanel::ThreadView_OnViewPositionYChanged( double PosStartY, d
 	VerticalScrollBar->SetState( (float)OffsetFraction, (float)ThumbSizeFraction );
 }
 
-void SProfilerGraphPanel::DataGraph_OnViewModeChanged( EDataGraphViewModes::Type ViewMode )
-{
-	if( ViewMode == EDataGraphViewModes::Index )
-	{
-		UpdateInternals();
-	}
-	else if( ViewMode == EDataGraphViewModes::Time )
-	{
-		// @TODO yrx 2014-04-23 
-	}
-}
-
-
 void SProfilerGraphPanel::ProfilerManager_OnViewModeChanged( EProfilerViewMode::Type NewViewMode )
 {
 	if( NewViewMode == EProfilerViewMode::LineIndexBased )
@@ -261,18 +244,18 @@ void SProfilerGraphPanel::UpdateInternals()
 {
 	if( ViewMode == EProfilerViewMode::LineIndexBased )
 	{
-	NumVisiblePoints = DataGraph->GetNumVisiblePoints();
-	NumDataPoints = DataGraph->GetNumDataPoints();
+		NumVisiblePoints = DataGraph->GetNumVisiblePoints();
+		NumDataPoints = DataGraph->GetNumDataPoints();
 
-	SetScrollBarState();
+		SetScrollBarState();
 		ProfilerMiniView.Pin()->MoveWithoutZoomSelectionBox( GraphOffset, GraphOffset + NumVisiblePoints );
-	 
-	if( FProfilerManager::Get()->IsLivePreview() )
-	{
-		// Scroll to the end.
-		HorizontalScrollBar_OnUserScrolled(1.0f);
+
+		if( FProfilerManager::Get()->IsLivePreview() )
+		{
+			// Scroll to the end.
+			HorizontalScrollBar_OnUserScrolled( 1.0f );
+		}
 	}
-}
 	else if( ViewMode == EProfilerViewMode::ThreadViewTimeBased )
 	{
 		// @TODO yrx 2014-04-23 
diff --git a/Engine/Source/Developer/Profiler/Private/Widgets/SProfilerGraphPanel.h b/Engine/Source/Developer/Profiler/Private/Widgets/SProfilerGraphPanel.h
--- a/Engine/Source/Developer/Profiler/Private/Widgets/SProfilerGraphPanel.h
+++ b/Engine/Source/Developer/Profiler/Private/Widgets/SProfilerGraphPanel.h
@@ -60,8 +60,6 @@ protected:
 	/** Sets state of the scroll bar. */
 	void SetScrollBarState();
 
-	/** Called when the data graph view mode has changed from the data graph widget. */
-	void DataGraph_OnViewModeChanged( EDataGraphViewModes::Type ViewMode );
 
 	void UpdateInternals();
 
