Slate - Renamed FSlateRect.GetPosition to GetTopLeft (as per NickA's request)

diff --git a/Engine/Source/Developer/Profiler/Private/Widgets/SProfilerThreadView.cpp b/Engine/Source/Developer/Profiler/Private/Widgets/SProfilerThreadView.cpp
--- a/Engine/Source/Developer/Profiler/Private/Widgets/SProfilerThreadView.cpp
+++ b/Engine/Source/Developer/Profiler/Private/Widgets/SProfilerThreadView.cpp
@@ -192,7 +192,7 @@ void SProfilerThreadView::DrawFramesBackgroundAndTimelines() const
 			(
 				PaintState->OutDrawElements,
 				PaintState->LayerId,
-				PaintState->AllottedGeometry.ToPaintGeometry( ClippedFrameBackgroundRect.GetPosition(), ClippedFrameBackgroundRect.GetSize() ),
+				PaintState->AllottedGeometry.ToPaintGeometry( ClippedFrameBackgroundRect.GetTopLeft(), ClippedFrameBackgroundRect.GetSize() ),
 				&SolidWhiteBrush,
 				PaintState->AbsoluteClippingRect,
 				PaintState->DrawEffects,
@@ -292,7 +292,7 @@ void SProfilerThreadView::DrawUIStackNodes() const
 				(
 					PaintState->OutDrawElements,
 					PaintState->LayerId,
-					PaintState->AllottedGeometry.ToPaintGeometry( ClippedNodeRect.GetPosition(), ClippedNodeRect.GetSize() ),
+					PaintState->AllottedGeometry.ToPaintGeometry( ClippedNodeRect.GetTopLeft(), ClippedNodeRect.GetSize() ),
 					BorderBrush,
 					PaintState->AbsoluteClippingRect,
 					PaintState->DrawEffects,
@@ -354,7 +354,7 @@ void SProfilerThreadView::DrawUIStackNodes() const
 				AdjustedPositionPx = FVector2D( ClippedNodeRect.Left, PositionPx.Y + MarkerPosYOffsetPx );
 			}
 
-			const FVector2D AbsolutePositionPx = PaintState->AllottedGeometry.LocalToAbsolute( ClippedNodeRect.GetPosition() );
+			const FVector2D AbsolutePositionPx = PaintState->AllottedGeometry.LocalToAbsolute( ClippedNodeRect.GetTopLeft() );
 			const FSlateRect AbsoluteClippingRect = FSlateRect( AbsolutePositionPx, AbsolutePositionPx + ClippedNodeRect.GetSize() );
 			
 			DrawText( bUseShortVersion ? StringStatName : StringStatNameWithTime, PaintState->SummaryFont8, AdjustedPositionPx, FColorList::White, FColorList::Black, FVector2D( 1.0f, 1.0f ), &AbsoluteClippingRect );
diff --git a/Engine/Source/Runtime/SlateCore/Public/Layout/SlateRect.h b/Engine/Source/Runtime/SlateCore/Public/Layout/SlateRect.h
--- a/Engine/Source/Runtime/SlateCore/Public/Layout/SlateRect.h
+++ b/Engine/Source/Runtime/SlateCore/Public/Layout/SlateRect.h
@@ -42,7 +42,7 @@ public:
 
 
 	/**
-	 * @return true, if the rectangle is empty.
+	 * @return true, if the rectangle is empty. Should be used in conjunction with IntersectionWith.
 	 */
 	bool IsEmpty() const
 	{
@@ -70,11 +70,11 @@ public:
 	}
 
 	/**
-	 * Returns the left-top position of the rectangle
+	 * Returns the top-left position of the rectangle
 	 * 
-	 * @return The left-top position.
+	 * @return The top-left position.
 	 */
-	FVector2D GetPosition() const
+	FVector2D GetTopLeft() const
 	{
 		return FVector2D( Left, Top );
 	}
