Removing debugging helpers that confuse people.

Contributed by UE4Linux community, part of PR 140.

diff --git a/Engine/Source/Developer/Profiler/Private/ProfilerSession.cpp b/Engine/Source/Developer/Profiler/Private/ProfilerSession.cpp
--- a/Engine/Source/Developer/Profiler/Private/ProfilerSession.cpp
+++ b/Engine/Source/Developer/Profiler/Private/ProfilerSession.cpp
@@ -634,12 +634,6 @@ void FProfilerStatMetaData::UpdateFromStatsState( const FStatsThreadState& Stats
 {
 	TMap<FName, int32> GroupFNameIDs;
 
-	StatsThreadStats.Groups;
-	StatsThreadStats.MemoryPoolToCapacityLongName;
-	StatsThreadStats.NotClearedEveryFrame;
-	StatsThreadStats.Threads;
-	StatsThreadStats.ShortNameToLongName;
-
 	for( auto It = StatsThreadStats.Threads.CreateConstIterator(); It; ++It )
 	{
 		ThreadDescriptions.Add( It.Key(), It.Value().ToString() );
@@ -927,4 +921,4 @@ void FRawProfilerSession::ProcessStatPacketArray( const FStatPacketArray& StatPa
 	}
 
 	out_ProfilerFrame.SortChildren();
-}
\ No newline at end of file
+}
diff --git a/Engine/Source/Developer/Profiler/Private/ProfilerStream.cpp b/Engine/Source/Developer/Profiler/Private/ProfilerStream.cpp
--- a/Engine/Source/Developer/Profiler/Private/ProfilerStream.cpp
+++ b/Engine/Source/Developer/Profiler/Private/ProfilerStream.cpp
@@ -10,9 +10,6 @@ void FProfilerUIStream::GenerateUIStream( const FProfilerStream& ProfilerStream,
 {
 	//SCOPE_LOG_TIME_FUNC();
 
-	ProfilerStream;
-	ZoomFactorX;
-
 	// Just copy the data for now.
 	const FIntPoint FramesIndices = ProfilerStream.GetFramesIndicesForTimeRange( StartTimeMS, EndTimeMS );
 	const double FrameTimeRangeMS = EndTimeMS - StartTimeMS;
@@ -33,10 +30,6 @@ void FProfilerUIStream::GenerateUIStream( const FProfilerStream& ProfilerStream,
 		{
 			const FProfilerFrame* ProfilerFrame = ProfilerStream.GetProfilerFrame( FrameIndex );
 
-			ProfilerFrame->FrameTimeMS;
-			ProfilerFrame->LastAccessTime;
-			ProfilerFrame->ThreadTimesMS;
-
 			// Thread nodes.
 			const int32 NumThreads = ProfilerFrame->Root->Children.Num();
 			for( int32 ThreadIndex = 0; ThreadIndex < NumThreads; ++ThreadIndex )
diff --git a/Engine/Source/Developer/Profiler/Private/ProfilerStream.h b/Engine/Source/Developer/Profiler/Private/ProfilerStream.h
--- a/Engine/Source/Developer/Profiler/Private/ProfilerStream.h
+++ b/Engine/Source/Developer/Profiler/Private/ProfilerStream.h
@@ -396,7 +396,6 @@ struct FProfilerUIStackNode
 
 	~FProfilerUIStackNode()
 	{
-		OriginalStackNodes;
 	}
 	
 	void InitializeUIData( const double NumMillisecondsPerWindow, const double NumPixelsPerMillisecond, const double NumMillisecondsPerSample )
