Profiler - Disabled debug code

diff --git a/Engine/Source/Developer/Profiler/Private/Widgets/SProfilerThreadView.cpp b/Engine/Source/Developer/Profiler/Private/Widgets/SProfilerThreadView.cpp
--- a/Engine/Source/Developer/Profiler/Private/Widgets/SProfilerThreadView.cpp
+++ b/Engine/Source/Developer/Profiler/Private/Widgets/SProfilerThreadView.cpp
@@ -114,7 +114,7 @@ int32 SProfilerThreadView::OnPaint( const FGeometry& AllottedGeometry, const FSl
 		DrawFrameMarkers();
 	}
 
-#if 1/*DEBUG_PROFILER_PERFORMANCE*/
+#if 0/*DEBUG_PROFILER_PERFORMANCE*/
 
 	LayerId++;
 	// Draw debug information.
