Profiler - Fixed a crash that may happen when loading a large stats file

diff --git a/Engine/Source/Developer/Profiler/Private/Widgets/SProfilerMiniView.cpp b/Engine/Source/Developer/Profiler/Private/Widgets/SProfilerMiniView.cpp
--- a/Engine/Source/Developer/Profiler/Private/Widgets/SProfilerMiniView.cpp
+++ b/Engine/Source/Developer/Profiler/Private/Widgets/SProfilerMiniView.cpp
@@ -661,16 +661,13 @@ void SProfilerMiniView::ProcessData()
 		// Aggregate.
 		const float ScaleRatio = (float)NumMiniViewSamples / (float)NumAllFrames;
 
-		float CurrentSample = 0.0f;
 		for( int32 FrameIndex = 0; FrameIndex < NumAllFrames; ++FrameIndex )
 		{
 			const FFrameThreadTimes& FrameThreadTimes = AllFrames[FrameIndex];
-			const int32 SampleIndex = FMath::TruncToInt(CurrentSample);
+			const int32 SampleIndex = FMath::TruncToInt( ScaleRatio*FrameIndex );
 
 			FMiniViewSample& Dest = MiniViewSamples[SampleIndex];
 			Dest.AddFrameAndFindMax( FrameThreadTimes );
-			
-			CurrentSample += ScaleRatio;
 		}
 
 		for( FMiniViewSample& It : MiniViewSamples )
