Profiler - Raw stats files should use stored SecondsPerCycle, instead of the current one

diff --git a/Engine/Source/Developer/Profiler/Private/ProfilerSession.cpp b/Engine/Source/Developer/Profiler/Private/ProfilerSession.cpp
--- a/Engine/Source/Developer/Profiler/Private/ProfilerSession.cpp
+++ b/Engine/Source/Developer/Profiler/Private/ProfilerSession.cpp
@@ -530,6 +530,51 @@ void FRawProfilerSession::PrepareLoading()
 		// Read frames offsets.
 		Stream.ReadFramesOffsets( *FileReader );
 
+		// Temporary fix for setting the correct value for the SecondsPerCycle.
+		{
+			UE_LOG( LogStats, Log, TEXT( "Looking for the STAT_SecondsPerCycle" ) );
+			FileReader->Seek( Stream.FramesInfo[0].FrameFileOffset );
+
+			FStatPacketArray StatPacketArray;
+
+			int64 TargetFrame;
+			*FileReader << TargetFrame;
+
+			int32 NumPackets;
+			*FileReader << NumPackets;
+
+			for( int32 PacketIndex = 0; PacketIndex < NumPackets; PacketIndex++ )
+			{
+				FStatPacket* ToRead = new FStatPacket();
+				Stream.ReadStatPacket( *FileReader, *ToRead, bIsFinalized );
+				StatPacketArray.Packets.Add( ToRead );
+			}
+
+			for( int32 PacketIndex = 0; PacketIndex < StatPacketArray.Packets.Num(); PacketIndex++ )
+			{
+				const FStatPacket& StatPacket = *StatPacketArray.Packets[PacketIndex];
+				const TArray<FStatMessage>& Data = StatPacket.StatMessages;
+
+				for( int32 Index = 0; Index < Data.Num(); Index++ )
+				{
+					const FStatMessage& Item = Data[Index];
+					check( Item.NameAndInfo.GetFlag( EStatMetaFlags::DummyAlwaysOne ) );  // we should never be sending short names to the stats anymore
+
+					const FName ShortName = Item.NameAndInfo.GetShortName();
+					if( ShortName == TEXT( "STAT_SecondsPerCycle" ) )
+					{
+						StatMetaData->SecondsPerCycle = Item.GetValue_double();
+						UE_LOG( LogStats, Log, TEXT( "STAT_SecondsPerCycle is %f [ns]" ), StatMetaData->GetSecondsPerCycle()*1000*1000 );
+						
+						PacketIndex = StatPacketArray.Packets.Num();
+						break;
+					}
+				}
+			}
+		}
+
+		check( StatMetaData->GetSecondsPerCycle() > 0.0 );
+
 		// Prepare profiler frames.
 		double ElapsedTimeMS = 0;
 		for( int32 FrameIndex = 0; FrameIndex < Stream.FramesInfo.Num(); ++FrameIndex )
@@ -739,14 +784,7 @@ void FProfilerStatMetaData::UpdateFromStatsState( const FStatsThreadState& Stats
 				RenderThreadIDs.AddUnique( ThreadID );
 			}
 		}
-
-		if( StatName == TEXT( "STAT_SecondsPerCycle" ) )
-		{
-			SecondsPerCycle = LongName.GetValue_double();
-		}
 	}
-
-	SecondsPerCycle = FPlatformTime::GetSecondsPerCycle();
 }
 
 void FRawProfilerSession::ProcessStatPacketArray( const FStatPacketArray& StatPacketArray, FProfilerFrame& out_ProfilerFrame, int32 FrameIndex )
