Profiler - Disabled capture button because current implementation in the profiler service is not thread-safe

diff --git a/Engine/Source/Developer/Profiler/Private/ProfilerCommands.cpp b/Engine/Source/Developer/Profiler/Private/ProfilerCommands.cpp
--- a/Engine/Source/Developer/Profiler/Private/ProfilerCommands.cpp
+++ b/Engine/Source/Developer/Profiler/Private/ProfilerCommands.cpp
@@ -408,6 +408,8 @@ void FProfilerActionManager::ToggleDataCapture_Execute( const FGuid SessionInsta
 
 bool FProfilerActionManager::ToggleDataCapture_CanExecute( const FGuid SessionInstanceID ) const
 {
+	// Disabled capture button because current implementation in the profiler service is not thread-safe.
+#if	0
 	// One session instance
 	if( SessionInstanceID.IsValid() )
 	{
@@ -420,6 +422,7 @@ bool FProfilerActionManager::ToggleDataCapture_CanExecute( const FGuid SessionIn
 		const bool bCanExecute = This->ActiveSession.IsValid() && This->ProfilerType == EProfilerSessionTypes::Live && This->GetProfilerInstancesNum() > 0;
 		return bCanExecute;
 	}
+#endif // 0
 
 	return false;
 }
