Profiler - Fixed an issue introduced by the recent shadow variable fixes

diff --git a/Engine/Source/Developer/Profiler/Private/ProfilerSession.cpp b/Engine/Source/Developer/Profiler/Private/ProfilerSession.cpp
--- a/Engine/Source/Developer/Profiler/Private/ProfilerSession.cpp
+++ b/Engine/Source/Developer/Profiler/Private/ProfilerSession.cpp
@@ -760,7 +760,7 @@ void FProfilerStatMetaData::UpdateFromStatsState( const FStatsThreadState& Stats
 		if( GroupName == FStatConstants::NAME_ThreadGroup )
 		{
 			uint32 ThreadID = 0;
-			for( auto ThreadsIt = StatsThreadStats.Threads.CreateConstIterator(); It; ++It )
+			for( auto ThreadsIt = StatsThreadStats.Threads.CreateConstIterator(); ThreadsIt; ++ThreadsIt )
 			{
 				if (ThreadsIt.Value() == StatName)
 				{
