UE4 Refactoring. Changed OVERRIDE and FINAL macros to keywords override and final.

diff --git a/Engine/Source/Developer/MovieSceneTools/Private/TransformTrackEditor.cpp b/Engine/Source/Developer/MovieSceneTools/Private/TransformTrackEditor.cpp
--- a/Engine/Source/Developer/MovieSceneTools/Private/TransformTrackEditor.cpp
+++ b/Engine/Source/Developer/MovieSceneTools/Private/TransformTrackEditor.cpp
@@ -26,19 +26,19 @@ public:
 	}
 
 	/** ISequencerSection interface */
-	virtual UMovieSceneSection* GetSectionObject() OVERRIDE
+	virtual UMovieSceneSection* GetSectionObject() override
 	{ 
 		return &Section;
 	}
 
-	virtual FString GetDisplayName() const OVERRIDE
+	virtual FString GetDisplayName() const override
 	{ 
 		return TEXT("Transform");
 	}
 	
-	virtual FString GetSectionTitle() const OVERRIDE { return FString(); }
+	virtual FString GetSectionTitle() const override { return FString(); }
 
-	virtual void GenerateSectionLayout( class ISectionLayoutBuilder& LayoutBuilder ) const OVERRIDE
+	virtual void GenerateSectionLayout( class ISectionLayoutBuilder& LayoutBuilder ) const override
 	{
 		UMovieSceneTransformSection* TransformSection = Cast<UMovieSceneTransformSection>( &Section );
 
@@ -62,7 +62,7 @@ public:
 		LayoutBuilder.PopCategory();
 	}
 
-	virtual int32 OnPaintSection( const FGeometry& AllottedGeometry, const FSlateRect& SectionClippingRect, FSlateWindowElementList& OutDrawElements, int32 LayerId, bool bParentEnabled ) const OVERRIDE 
+	virtual int32 OnPaintSection( const FGeometry& AllottedGeometry, const FSlateRect& SectionClippingRect, FSlateWindowElementList& OutDrawElements, int32 LayerId, bool bParentEnabled ) const override 
 	{
 		// Add a box for the section
 		FSlateDrawElement::MakeBox( 
diff --git a/Engine/Source/Developer/MovieSceneTools/Private/TransformTrackEditor.h b/Engine/Source/Developer/MovieSceneTools/Private/TransformTrackEditor.h
--- a/Engine/Source/Developer/MovieSceneTools/Private/TransformTrackEditor.h
+++ b/Engine/Source/Developer/MovieSceneTools/Private/TransformTrackEditor.h
@@ -27,9 +27,9 @@ public:
 	static TSharedRef<FMovieSceneTrackEditor> CreateTrackEditor( TSharedRef<ISequencer> OwningSequencer );
 
 	/** FMovieSceneTrackEditor Interface */
-	virtual bool SupportsType( TSubclassOf<UMovieSceneTrack> Type ) const OVERRIDE;
-	virtual TSharedRef<ISequencerSection> MakeSectionInterface( UMovieSceneSection& SectionObject, UMovieSceneTrack* Track ) OVERRIDE;
-	virtual void AddKey( const FGuid& ObjectGuid, UObject* AdditionalAsset = NULL ) OVERRIDE;
+	virtual bool SupportsType( TSubclassOf<UMovieSceneTrack> Type ) const override;
+	virtual TSharedRef<ISequencerSection> MakeSectionInterface( UMovieSceneSection& SectionObject, UMovieSceneTrack* Track ) override;
+	virtual void AddKey( const FGuid& ObjectGuid, UObject* AdditionalAsset = NULL ) override;
 
 private:
 	/**
diff --git a/Engine/Source/Developer/OutputLog/Private/OutputLogModule.cpp b/Engine/Source/Developer/OutputLog/Private/OutputLogModule.cpp
--- a/Engine/Source/Developer/OutputLog/Private/OutputLogModule.cpp
+++ b/Engine/Source/Developer/OutputLog/Private/OutputLogModule.cpp
@@ -41,7 +41,7 @@ public:
 
 protected:
 
-	virtual void Serialize( const TCHAR* V, ELogVerbosity::Type Verbosity, const class FName& Category ) OVERRIDE
+	virtual void Serialize( const TCHAR* V, ELogVerbosity::Type Verbosity, const class FName& Category ) override
 	{
 		// Capture all incoming messages and store them in history
 		SOutputLog::CreateLogMessages(V, Verbosity, Category, Messages);
diff --git a/Engine/Source/Developer/OutputLog/Private/SOutputLog.h b/Engine/Source/Developer/OutputLog/Private/SOutputLog.h
--- a/Engine/Source/Developer/OutputLog/Private/SOutputLog.h
+++ b/Engine/Source/Developer/OutputLog/Private/SOutputLog.h
@@ -55,7 +55,7 @@ public:
 	}
 
 	/** SWidget interface */
-	virtual void Tick( const FGeometry& AllottedGeometry, const double InCurrentTime, const float InDeltaTime ) OVERRIDE;
+	virtual void Tick( const FGeometry& AllottedGeometry, const double InCurrentTime, const float InDeltaTime ) override;
 
 protected:
 
@@ -156,11 +156,11 @@ public:
 	 *
 	 * @return  Returns whether the event was handled, along with other possible actions
 	 */
-	virtual FReply OnKeyDown( const FGeometry& MyGeometry, const FKeyboardEvent& InKeyboardEvent ) OVERRIDE;
+	virtual FReply OnKeyDown( const FGeometry& MyGeometry, const FKeyboardEvent& InKeyboardEvent ) override;
 
 protected:
 
-	virtual void Serialize( const TCHAR* V, ELogVerbosity::Type Verbosity, const class FName& Category ) OVERRIDE;
+	virtual void Serialize( const TCHAR* V, ELogVerbosity::Type Verbosity, const class FName& Category ) override;
 
 	/** Makes the widget for the log messages in the list view */
 	TSharedRef<ITableRow> MakeLogListItemWidget(TSharedPtr<FLogMessage> Message, const TSharedRef<STableViewBase>& OwnerTable);
diff --git a/Engine/Source/Developer/PhysXFormats/Private/PhysXFormats.cpp b/Engine/Source/Developer/PhysXFormats/Private/PhysXFormats.cpp
--- a/Engine/Source/Developer/PhysXFormats/Private/PhysXFormats.cpp
+++ b/Engine/Source/Developer/PhysXFormats/Private/PhysXFormats.cpp
@@ -67,7 +67,7 @@ public:
 		return false;
 	}
 
-	virtual uint16 GetVersion(FName Format) const OVERRIDE
+	virtual uint16 GetVersion(FName Format) const override
 	{
 		check(CheckPhysXFormat(Format));
 		return UE_PHYSX_PC_VER;
@@ -81,7 +81,7 @@ public:
 		OutFormats.Add(NAME_PhysXPS4);
 	}
 
-	virtual bool CookConvex(FName Format, const TArray<FVector>& SrcBuffer, TArray<uint8>& OutBuffer) const OVERRIDE
+	virtual bool CookConvex(FName Format, const TArray<FVector>& SrcBuffer, TArray<uint8>& OutBuffer) const override
 	{
 #if WITH_PHYSX
 		PxPlatform::Enum PhysXFormat = PxPlatform::ePC;
@@ -115,7 +115,7 @@ public:
 		return false;
 	}
 
-	virtual bool CookTriMesh(FName Format, const TArray<FVector>& SrcVertices, const TArray<FTriIndices>& SrcIndices, const TArray<uint16>& SrcMaterialIndices, const bool FlipNormals, TArray<uint8>& OutBuffer) const OVERRIDE
+	virtual bool CookTriMesh(FName Format, const TArray<FVector>& SrcVertices, const TArray<FTriIndices>& SrcIndices, const TArray<uint16>& SrcMaterialIndices, const bool FlipNormals, TArray<uint8>& OutBuffer) const override
 	{
 #if WITH_PHYSX
 		PxPlatform::Enum PhysXFormat = PxPlatform::ePC;
@@ -148,7 +148,7 @@ public:
 #endif		// WITH_PHYSX
 	}
 
-	virtual bool CookHeightField(FName Format, FIntPoint HFSize, float Thickness, const void* Samples, uint32 SamplesStride, TArray<uint8>& OutBuffer) const OVERRIDE
+	virtual bool CookHeightField(FName Format, FIntPoint HFSize, float Thickness, const void* Samples, uint32 SamplesStride, TArray<uint8>& OutBuffer) const override
 	{
 #if WITH_PHYSX
 		PxPlatform::Enum PhysXFormat = PxPlatform::ePC;
diff --git a/Engine/Source/Developer/PhysXFormats/Private/PhysXFormats.h b/Engine/Source/Developer/PhysXFormats/Private/PhysXFormats.h
--- a/Engine/Source/Developer/PhysXFormats/Private/PhysXFormats.h
+++ b/Engine/Source/Developer/PhysXFormats/Private/PhysXFormats.h
@@ -64,7 +64,7 @@ public:
 		: Data(InData)
 	{}
 
-	virtual PxU32 write(const void* Src, PxU32 Count) OVERRIDE
+	virtual PxU32 write(const void* Src, PxU32 Count) override
 	{
 		check(Data);
 		check(Src);
diff --git a/Engine/Source/Developer/Profiler/Private/ProfilerDataProvider.h b/Engine/Source/Developer/Profiler/Private/ProfilerDataProvider.h
--- a/Engine/Source/Developer/Profiler/Private/ProfilerDataProvider.h
+++ b/Engine/Source/Developer/Profiler/Private/ProfilerDataProvider.h
@@ -243,7 +243,7 @@ public:
 		const double InDurationMS, 
 		const uint32 InCallsPerFrame, 
 		const uint32 InParentIndex = FProfilerSample::InvalidIndex 
-	) OVERRIDE;
+	) override;
 
 	virtual void AddCounterSample
 	( 
@@ -251,15 +251,15 @@ public:
 		const uint32 InStatID, 
 		const double InCounter, 
 		const EProfilerSampleTypes::Type InProfilerSampleType 
-	) OVERRIDE;
+	) override;
 
-	virtual const uint32 AddDuplicatedSample( const FProfilerSample& ProfilerSample ) OVERRIDE;
+	virtual const uint32 AddDuplicatedSample( const FProfilerSample& ProfilerSample ) override;
 
-	virtual const uint32 GetNumSamples() const OVERRIDE;
+	virtual const uint32 GetNumSamples() const override;
 
-	virtual const uint64 GetMemoryUsage() const OVERRIDE;;
+	virtual const uint64 GetMemoryUsage() const override;;
 
-	virtual const FProfilerSampleArray& GetCollection() const OVERRIDE
+	virtual const FProfilerSampleArray& GetCollection() const override
 	{
 		return Collection;
 	}
diff --git a/Engine/Source/Developer/Profiler/Private/ProfilerFPSAnalyzer.h b/Engine/Source/Developer/Profiler/Private/ProfilerFPSAnalyzer.h
--- a/Engine/Source/Developer/Profiler/Private/ProfilerFPSAnalyzer.h
+++ b/Engine/Source/Developer/Profiler/Private/ProfilerFPSAnalyzer.h
@@ -45,9 +45,9 @@ public:
 	virtual ~FFPSAnalyzer()
 	{}
 
-	virtual int32 GetCount(float MinVal, float MaxVal) OVERRIDE;
+	virtual int32 GetCount(float MinVal, float MaxVal) override;
 
-	virtual int32 GetTotalCount() OVERRIDE;
+	virtual int32 GetTotalCount() override;
 
 	void Reset();
 
diff --git a/Engine/Source/Developer/Profiler/Private/ProfilerModule.cpp b/Engine/Source/Developer/Profiler/Private/ProfilerModule.cpp
--- a/Engine/Source/Developer/Profiler/Private/ProfilerModule.cpp
+++ b/Engine/Source/Developer/Profiler/Private/ProfilerModule.cpp
@@ -13,7 +13,7 @@ class FProfilerModule
 	: public IProfilerModule
 {
 public:
-	virtual TSharedRef<SWidget> CreateProfilerWindow( const ISessionManagerRef& InSessionManager, const TSharedRef<SDockTab>& ConstructUnderMajorTab ) OVERRIDE
+	virtual TSharedRef<SWidget> CreateProfilerWindow( const ISessionManagerRef& InSessionManager, const TSharedRef<SDockTab>& ConstructUnderMajorTab ) override
 	{
 		ProfilerManager = FProfilerManager::Initialize( InSessionManager );
 		TSharedRef<SProfilerWindow> ProfilerWindow = SNew(SProfilerWindow);
@@ -24,10 +24,10 @@ public:
 		return ProfilerWindow;
 	}
 
-	virtual void StartupModule() OVERRIDE
+	virtual void StartupModule() override
 	{}
 
-	virtual void ShutdownModule() OVERRIDE
+	virtual void ShutdownModule() override
 	{
 		if( FProfilerManager::Get().IsValid() )
 		{
@@ -35,12 +35,12 @@ public:
 		}
 	}
 
-	virtual bool SupportsDynamicReloading() OVERRIDE
+	virtual bool SupportsDynamicReloading() override
 	{
 		return false;
 	}
 
-	virtual TWeakPtr<class IProfilerManager> GetProfilerManager() OVERRIDE
+	virtual TWeakPtr<class IProfilerManager> GetProfilerManager() override
 	{
 		return ProfilerManager;
 	}
diff --git a/Engine/Source/Developer/Profiler/Private/Widgets/SDataGraph.h b/Engine/Source/Developer/Profiler/Private/Widgets/SDataGraph.h
--- a/Engine/Source/Developer/Profiler/Private/Widgets/SDataGraph.h
+++ b/Engine/Source/Developer/Profiler/Private/Widgets/SDataGraph.h
@@ -136,7 +136,7 @@ public:
 	 * @param  InCurrentTime  Current absolute real time
 	 * @param  InDeltaTime  Real time passed since last tick
 	 */
-	virtual void Tick( const FGeometry& AllottedGeometry, const double InCurrentTime, const float InDeltaTime ) OVERRIDE;
+	virtual void Tick( const FGeometry& AllottedGeometry, const double InCurrentTime, const float InDeltaTime ) override;
 
 	/**
 	 * The widget should respond by populating the OutDrawElements array with FDrawElements 
@@ -151,7 +151,7 @@ public:
 	 *
 	 * @return The maximum layer ID attained by this widget or any of its children.
 	 */
-	virtual int32 OnPaint( const FGeometry& AllottedGeometry, const FSlateRect& MyClippingRect, FSlateWindowElementList& OutDrawElements, int32 LayerId, const FWidgetStyle& InWidgetStyle, bool bParentEnabled ) const OVERRIDE;
+	virtual int32 OnPaint( const FGeometry& AllottedGeometry, const FSlateRect& MyClippingRect, FSlateWindowElementList& OutDrawElements, int32 LayerId, const FWidgetStyle& InWidgetStyle, bool bParentEnabled ) const override;
 
 	/**
 	 * The system calls this method to notify the widget that a mouse button was pressed within it. This event is bubbled.
@@ -161,7 +161,7 @@ public:
 	 *
 	 * @return Whether the event was handled along with possible requests for the system to take action.
 	 */
-	virtual FReply OnMouseButtonDown( const FGeometry& MyGeometry, const FPointerEvent& MouseEvent ) OVERRIDE;
+	virtual FReply OnMouseButtonDown( const FGeometry& MyGeometry, const FPointerEvent& MouseEvent ) override;
 	
 	/**
 	 * The system calls this method to notify the widget that a mouse button was release within it. This event is bubbled.
@@ -171,7 +171,7 @@ public:
 	 *
 	 * @return Whether the event was handled along with possible requests for the system to take action.
 	 */
-	virtual FReply OnMouseButtonUp( const FGeometry& MyGeometry, const FPointerEvent& MouseEvent ) OVERRIDE;
+	virtual FReply OnMouseButtonUp( const FGeometry& MyGeometry, const FPointerEvent& MouseEvent ) override;
 	
 	/**
 	 * The system calls this method to notify the widget that a mouse moved within it. This event is bubbled.
@@ -181,7 +181,7 @@ public:
 	 *
 	 * @return Whether the event was handled along with possible requests for the system to take action.
 	 */
-	virtual FReply OnMouseMove( const FGeometry& MyGeometry, const FPointerEvent& MouseEvent ) OVERRIDE;
+	virtual FReply OnMouseMove( const FGeometry& MyGeometry, const FPointerEvent& MouseEvent ) override;
 	
 	/**
 	 * The system will use this event to notify a widget that the cursor has entered it. This event is NOT bubbled.
@@ -189,14 +189,14 @@ public:
 	 * @param MyGeometry The Geometry of the widget receiving the event
 	 * @param MouseEvent Information about the input event
 	 */
-	virtual void OnMouseEnter( const FGeometry& MyGeometry, const FPointerEvent& MouseEvent ) OVERRIDE;
+	virtual void OnMouseEnter( const FGeometry& MyGeometry, const FPointerEvent& MouseEvent ) override;
 	
 	/**
 	 * The system will use this event to notify a widget that the cursor has left it. This event is NOT bubbled.
 	 *
 	 * @param MouseEvent Information about the input event
 	 */
-	virtual void OnMouseLeave( const FPointerEvent& MouseEvent ) OVERRIDE;
+	virtual void OnMouseLeave( const FPointerEvent& MouseEvent ) override;
 	
 	/**
 	 * Called when the mouse wheel is spun. This event is bubbled.
@@ -205,7 +205,7 @@ public:
 	 *
 	 * @return  Returns whether the event was handled, along with other possible actions
 	 */
-	virtual FReply OnMouseWheel( const FGeometry& MyGeometry, const FPointerEvent& MouseEvent ) OVERRIDE;
+	virtual FReply OnMouseWheel( const FGeometry& MyGeometry, const FPointerEvent& MouseEvent ) override;
 
 	/**
 	 * Called when a mouse button is double clicked.  Override this in derived classes.
@@ -215,7 +215,7 @@ public:
 	 *
 	 * @return  Returns whether the event was handled, along with other possible actions
 	 */
-	virtual FReply OnMouseButtonDoubleClick( const FGeometry& MyGeometry, const FPointerEvent& MouseEvent ) OVERRIDE;
+	virtual FReply OnMouseButtonDoubleClick( const FGeometry& MyGeometry, const FPointerEvent& MouseEvent ) override;
 
 	/**
 	 * Called when the user is dropping something onto a widget; terminates drag and drop.
@@ -225,7 +225,7 @@ public:
 	 *
 	 * @return A reply that indicated whether this event was handled.
 	 */
-	virtual FReply OnDrop( const FGeometry& MyGeometry, const FDragDropEvent& DragDropEvent ) OVERRIDE;
+	virtual FReply OnDrop( const FGeometry& MyGeometry, const FDragDropEvent& DragDropEvent ) override;
 
 	/**
 	 * Called during drag and drop when the drag enters a widget.
@@ -244,14 +244,14 @@ public:
 	 *
 	 * @return A reply that indicated whether the contents of the DragDropEvent can potentially be processed by this widget.
 	 */
-	virtual void OnDragEnter( const FGeometry& MyGeometry, const FDragDropEvent& DragDropEvent )  OVERRIDE;
+	virtual void OnDragEnter( const FGeometry& MyGeometry, const FDragDropEvent& DragDropEvent )  override;
 	
 	/**
 	 * Called during drag and drop when the drag leaves a widget.
 	 *
 	 * @param DragDropEvent   The drag and drop event.
 	 */
-	virtual void OnDragLeave( const FDragDropEvent& DragDropEvent )  OVERRIDE;
+	virtual void OnDragLeave( const FDragDropEvent& DragDropEvent )  override;
 
 	/**
 	 * Called during drag and drop when the the mouse is being dragged over a widget.
@@ -261,14 +261,14 @@ public:
 	 *
 	 * @return A reply that indicated whether this event was handled.
 	 */
-	virtual FReply OnDragOver( const FGeometry& MyGeometry, const FDragDropEvent& DragDropEvent )  OVERRIDE;
+	virtual FReply OnDragOver( const FGeometry& MyGeometry, const FDragDropEvent& DragDropEvent )  override;
 
 	/**
 	 * Called when the system wants to know which cursor to display for this Widget.  This event is bubbled.
 	 *
 	 * @return  The cursor requested (can be None.)
 	 */
-	virtual FCursorReply OnCursorQuery( const FGeometry& MyGeometry, const FPointerEvent& CursorEvent ) const OVERRIDE;
+	virtual FCursorReply OnCursorQuery( const FGeometry& MyGeometry, const FPointerEvent& CursorEvent ) const override;
 
 
 	/*-----------------------------------------------------------------------------
@@ -337,7 +337,7 @@ public:
 	
 
 protected:
-	virtual FVector2D ComputeDesiredSize() const OVERRIDE
+	virtual FVector2D ComputeDesiredSize() const override
 	{
 		return FVector2D(16.0f,16.0f);
 	}
diff --git a/Engine/Source/Developer/Profiler/Private/Widgets/SEventGraph.cpp b/Engine/Source/Developer/Profiler/Private/Widgets/SEventGraph.cpp
--- a/Engine/Source/Developer/Profiler/Private/Widgets/SEventGraph.cpp
+++ b/Engine/Source/Developer/Profiler/Private/Widgets/SEventGraph.cpp
@@ -1,4 +1,4 @@
-﻿// Copyright 1998-2014 Epic Games, Inc. All Rights Reserved.
+// Copyright 1998-2014 Epic Games, Inc. All Rights Reserved.
 
 #include "ProfilerPrivatePCH.h"
 #include "SEventGraphTooltip.h"
@@ -450,7 +450,7 @@ protected:
 	 * @param MyGeometry The Geometry of the widget receiving the event
 	 * @param MouseEvent Information about the input event
 	 */
-	virtual void OnMouseEnter( const FGeometry& MyGeometry, const FPointerEvent& MouseEvent ) OVERRIDE
+	virtual void OnMouseEnter( const FGeometry& MyGeometry, const FPointerEvent& MouseEvent ) override
 	{
 		SCompoundWidget::OnMouseEnter( MyGeometry, MouseEvent );
 		SetHoveredTableCellDelegate.ExecuteIfBound( ColumnID, EventPtr );
@@ -461,7 +461,7 @@ protected:
 	 *
 	 * @param MouseEvent Information about the input event
 	 */
-	virtual void OnMouseLeave( const FPointerEvent& MouseEvent ) OVERRIDE
+	virtual void OnMouseLeave( const FPointerEvent& MouseEvent ) override
 	{
 		SCompoundWidget::OnMouseLeave( MouseEvent );
 		SetHoveredTableCellDelegate.ExecuteIfBound( NAME_None, nullptr );
@@ -484,7 +484,7 @@ protected:
 	 *
 	 * @return A reply that indicated whether the contents of the DragDropEvent can potentially be processed by this widget.
 	 */
-	virtual void OnDragEnter( const FGeometry& MyGeometry, const FDragDropEvent& DragDropEvent ) OVERRIDE
+	virtual void OnDragEnter( const FGeometry& MyGeometry, const FDragDropEvent& DragDropEvent ) override
 	{
 		SCompoundWidget::OnDragEnter( MyGeometry, DragDropEvent );
 		SetHoveredTableCellDelegate.ExecuteIfBound( ColumnID, EventPtr );
@@ -495,7 +495,7 @@ protected:
 	 *
 	 * @param DragDropEvent   The drag and drop event.
 	 */
-	virtual void OnDragLeave( const FDragDropEvent& DragDropEvent )  OVERRIDE
+	virtual void OnDragLeave( const FDragDropEvent& DragDropEvent )  override
 	{
 		SCompoundWidget::OnDragLeave( DragDropEvent );
 		SetHoveredTableCellDelegate.ExecuteIfBound( NAME_None, nullptr );
@@ -591,7 +591,7 @@ public:
 	 * @return a widget to represent the contents of a cell in this row of a TableView. 
 	 */
 	BEGIN_SLATE_FUNCTION_BUILD_OPTIMIZATION
-	virtual TSharedRef<SWidget> GenerateWidgetForColumn( const FName& ColumnID ) OVERRIDE
+	virtual TSharedRef<SWidget> GenerateWidgetForColumn( const FName& ColumnID ) override
 	{
 		return
 
@@ -640,7 +640,7 @@ public:
 	 * @param  InMouseEvent  MouseMove that triggered the drag
 	 * 
 	 */
-	virtual FReply OnDragDetected( const FGeometry& MyGeometry, const FPointerEvent& MouseEvent ) OVERRIDE
+	virtual FReply OnDragDetected( const FGeometry& MyGeometry, const FPointerEvent& MouseEvent ) override
 	{
 		if( MouseEvent.IsMouseButtonDown( EKeys::LeftMouseButton ) )
 		{
diff --git a/Engine/Source/Developer/Profiler/Private/Widgets/SFiltersAndPresets.cpp b/Engine/Source/Developer/Profiler/Private/Widgets/SFiltersAndPresets.cpp
--- a/Engine/Source/Developer/Profiler/Private/Widgets/SFiltersAndPresets.cpp
+++ b/Engine/Source/Developer/Profiler/Private/Widgets/SFiltersAndPresets.cpp
@@ -453,7 +453,7 @@ public:
 	 * @param  InMouseEvent  MouseMove that triggered the drag
 	 * 
 	 */
-	virtual FReply OnDragDetected( const FGeometry& MyGeometry, const FPointerEvent& MouseEvent ) OVERRIDE
+	virtual FReply OnDragDetected( const FGeometry& MyGeometry, const FPointerEvent& MouseEvent ) override
 	{
 		if( MouseEvent.IsMouseButtonDown( EKeys::LeftMouseButton ))
 		{
diff --git a/Engine/Source/Developer/Profiler/Private/Widgets/SHistogram.h b/Engine/Source/Developer/Profiler/Private/Widgets/SHistogram.h
--- a/Engine/Source/Developer/Profiler/Private/Widgets/SHistogram.h
+++ b/Engine/Source/Developer/Profiler/Private/Widgets/SHistogram.h
@@ -122,7 +122,7 @@ public:
 	 * @param  InCurrentTime  Current absolute real time
 	 * @param  InDeltaTime  Real time passed since last tick
 	 */
-	virtual void Tick( const FGeometry& AllottedGeometry, const double InCurrentTime, const float InDeltaTime ) OVERRIDE;
+	virtual void Tick( const FGeometry& AllottedGeometry, const double InCurrentTime, const float InDeltaTime ) override;
 
 	/**
 	 * The widget should respond by populating the OutDrawElements array with FDrawElements 
@@ -137,7 +137,7 @@ public:
 	 *
 	 * @return The maximum layer ID attained by this widget or any of its children.
 	 */
-	virtual int32 OnPaint( const FGeometry& AllottedGeometry, const FSlateRect& MyClippingRect, FSlateWindowElementList& OutDrawElements, int32 LayerId, const FWidgetStyle& InWidgetStyle, bool bParentEnabled ) const OVERRIDE;
+	virtual int32 OnPaint( const FGeometry& AllottedGeometry, const FSlateRect& MyClippingRect, FSlateWindowElementList& OutDrawElements, int32 LayerId, const FWidgetStyle& InWidgetStyle, bool bParentEnabled ) const override;
 
 	/**
 	 * The system calls this method to notify the widget that a mouse button was pressed within it. This event is bubbled.
@@ -147,7 +147,7 @@ public:
 	 *
 	 * @return Whether the event was handled along with possible requests for the system to take action.
 	 */
-	virtual FReply OnMouseButtonDown( const FGeometry& MyGeometry, const FPointerEvent& MouseEvent ) OVERRIDE;
+	virtual FReply OnMouseButtonDown( const FGeometry& MyGeometry, const FPointerEvent& MouseEvent ) override;
 	
 	/**
 	 * The system calls this method to notify the widget that a mouse button was release within it. This event is bubbled.
@@ -157,7 +157,7 @@ public:
 	 *
 	 * @return Whether the event was handled along with possible requests for the system to take action.
 	 */
-	virtual FReply OnMouseButtonUp( const FGeometry& MyGeometry, const FPointerEvent& MouseEvent ) OVERRIDE;
+	virtual FReply OnMouseButtonUp( const FGeometry& MyGeometry, const FPointerEvent& MouseEvent ) override;
 	
 	/**
 	 * The system calls this method to notify the widget that a mouse moved within it. This event is bubbled.
@@ -167,7 +167,7 @@ public:
 	 *
 	 * @return Whether the event was handled along with possible requests for the system to take action.
 	 */
-	virtual FReply OnMouseMove( const FGeometry& MyGeometry, const FPointerEvent& MouseEvent ) OVERRIDE;
+	virtual FReply OnMouseMove( const FGeometry& MyGeometry, const FPointerEvent& MouseEvent ) override;
 	
 	/**
 	 * The system will use this event to notify a widget that the cursor has entered it. This event is NOT bubbled.
@@ -175,14 +175,14 @@ public:
 	 * @param MyGeometry The Geometry of the widget receiving the event
 	 * @param MouseEvent Information about the input event
 	 */
-	virtual void OnMouseEnter( const FGeometry& MyGeometry, const FPointerEvent& MouseEvent ) OVERRIDE;
+	virtual void OnMouseEnter( const FGeometry& MyGeometry, const FPointerEvent& MouseEvent ) override;
 	
 	/**
 	 * The system will use this event to notify a widget that the cursor has left it. This event is NOT bubbled.
 	 *
 	 * @param MouseEvent Information about the input event
 	 */
-	virtual void OnMouseLeave( const FPointerEvent& MouseEvent ) OVERRIDE;
+	virtual void OnMouseLeave( const FPointerEvent& MouseEvent ) override;
 	
 	/**
 	 * Called when the mouse wheel is spun. This event is bubbled.
@@ -191,7 +191,7 @@ public:
 	 *
 	 * @return  Returns whether the event was handled, along with other possible actions
 	 */
-	virtual FReply OnMouseWheel( const FGeometry& MyGeometry, const FPointerEvent& MouseEvent ) OVERRIDE;
+	virtual FReply OnMouseWheel( const FGeometry& MyGeometry, const FPointerEvent& MouseEvent ) override;
 
 	/**
 	 * Called when a mouse button is double clicked.  Override this in derived classes.
@@ -201,19 +201,19 @@ public:
 	 *
 	 * @return  Returns whether the event was handled, along with other possible actions
 	 */
-	virtual FReply OnMouseButtonDoubleClick( const FGeometry& MyGeometry, const FPointerEvent& MouseEvent ) OVERRIDE;
+	virtual FReply OnMouseButtonDoubleClick( const FGeometry& MyGeometry, const FPointerEvent& MouseEvent ) override;
 
 	/**
 	 * Called when the system wants to know which cursor to display for this Widget.  This event is bubbled.
 	 *
 	 * @return  The cursor requested (can be None.)
 	 */
-	virtual FCursorReply OnCursorQuery( const FGeometry& MyGeometry, const FPointerEvent& CursorEvent ) const OVERRIDE;
+	virtual FCursorReply OnCursorQuery( const FGeometry& MyGeometry, const FPointerEvent& CursorEvent ) const override;
 
 	/** Restores default state for the data graph by removing all inner graphs. */
 	void RestoreDefaultState();
 
-	virtual FVector2D ComputeDesiredSize() const OVERRIDE
+	virtual FVector2D ComputeDesiredSize() const override
 	{
 		//return SCompoundWidget::ComputeDesiredSize();
 		return FVector2D(16.0f,16.0f);
diff --git a/Engine/Source/Developer/Profiler/Private/Widgets/SProfilerMiniView.h b/Engine/Source/Developer/Profiler/Private/Widgets/SProfilerMiniView.h
--- a/Engine/Source/Developer/Profiler/Private/Widgets/SProfilerMiniView.h
+++ b/Engine/Source/Developer/Profiler/Private/Widgets/SProfilerMiniView.h
@@ -211,7 +211,7 @@ public:
 	* @param  InCurrentTime  Current absolute real time
 	* @param  InDeltaTime  Real time passed since last tick
 	*/
-	virtual void Tick( const FGeometry& AllottedGeometry, const double InCurrentTime, const float InDeltaTime ) OVERRIDE;
+	virtual void Tick( const FGeometry& AllottedGeometry, const double InCurrentTime, const float InDeltaTime ) override;
 
 	/**
 	* The widget should respond by populating the OutDrawElements array with FDrawElements
@@ -226,7 +226,7 @@ public:
 	*
 	* @return The maximum layer ID attained by this widget or any of its children.
 	*/
-	virtual int32 OnPaint( const FGeometry& AllottedGeometry, const FSlateRect& MyClippingRect, FSlateWindowElementList& OutDrawElements, int32 LayerId, const FWidgetStyle& InWidgetStyle, bool bParentEnabled ) const OVERRIDE;
+	virtual int32 OnPaint( const FGeometry& AllottedGeometry, const FSlateRect& MyClippingRect, FSlateWindowElementList& OutDrawElements, int32 LayerId, const FWidgetStyle& InWidgetStyle, bool bParentEnabled ) const override;
 
 	/**
 	* The system calls this method to notify the widget that a mouse button was pressed within it. This event is bubbled.
@@ -236,7 +236,7 @@ public:
 	*
 	* @return Whether the event was handled along with possible requests for the system to take action.
 	*/
-	virtual FReply OnMouseButtonDown( const FGeometry& MyGeometry, const FPointerEvent& MouseEvent ) OVERRIDE;
+	virtual FReply OnMouseButtonDown( const FGeometry& MyGeometry, const FPointerEvent& MouseEvent ) override;
 
 	/**
 	* The system calls this method to notify the widget that a mouse button was release within it. This event is bubbled.
@@ -246,7 +246,7 @@ public:
 	*
 	* @return Whether the event was handled along with possible requests for the system to take action.
 	*/
-	virtual FReply OnMouseButtonUp( const FGeometry& MyGeometry, const FPointerEvent& MouseEvent ) OVERRIDE;
+	virtual FReply OnMouseButtonUp( const FGeometry& MyGeometry, const FPointerEvent& MouseEvent ) override;
 
 	/**
 	* The system calls this method to notify the widget that a mouse moved within it. This event is bubbled.
@@ -256,7 +256,7 @@ public:
 	*
 	* @return Whether the event was handled along with possible requests for the system to take action.
 	*/
-	virtual FReply OnMouseMove( const FGeometry& MyGeometry, const FPointerEvent& MouseEvent ) OVERRIDE;
+	virtual FReply OnMouseMove( const FGeometry& MyGeometry, const FPointerEvent& MouseEvent ) override;
 
 	/**
 	* The system will use this event to notify a widget that the cursor has entered it. This event is NOT bubbled.
@@ -264,14 +264,14 @@ public:
 	* @param MyGeometry The Geometry of the widget receiving the event
 	* @param MouseEvent Information about the input event
 	*/
-	virtual void OnMouseEnter( const FGeometry& MyGeometry, const FPointerEvent& MouseEvent ) OVERRIDE;
+	virtual void OnMouseEnter( const FGeometry& MyGeometry, const FPointerEvent& MouseEvent ) override;
 
 	/**
 	* The system will use this event to notify a widget that the cursor has left it. This event is NOT bubbled.
 	*
 	* @param MouseEvent Information about the input event
 	*/
-	virtual void OnMouseLeave( const FPointerEvent& MouseEvent ) OVERRIDE;
+	virtual void OnMouseLeave( const FPointerEvent& MouseEvent ) override;
 
 	/**
 	* Called when the mouse wheel is spun. This event is bubbled.
@@ -280,7 +280,7 @@ public:
 	*
 	* @return  Returns whether the event was handled, along with other possible actions
 	*/
-	virtual FReply OnMouseWheel( const FGeometry& MyGeometry, const FPointerEvent& MouseEvent ) OVERRIDE;
+	virtual FReply OnMouseWheel( const FGeometry& MyGeometry, const FPointerEvent& MouseEvent ) override;
 
 	/**
 	* Called when a mouse button is double clicked.  Override this in derived classes.
@@ -290,14 +290,14 @@ public:
 	*
 	* @return  Returns whether the event was handled, along with other possible actions
 	*/
-	virtual FReply OnMouseButtonDoubleClick( const FGeometry& MyGeometry, const FPointerEvent& MouseEvent ) OVERRIDE;
+	virtual FReply OnMouseButtonDoubleClick( const FGeometry& MyGeometry, const FPointerEvent& MouseEvent ) override;
 
 	/**
 	* Called when the system wants to know which cursor to display for this Widget.  This event is bubbled.
 	*
 	* @return  The cursor requested (can be None.)
 	*/
-	virtual FCursorReply OnCursorQuery( const FGeometry& MyGeometry, const FPointerEvent& CursorEvent ) const OVERRIDE;
+	virtual FCursorReply OnCursorQuery( const FGeometry& MyGeometry, const FPointerEvent& CursorEvent ) const override;
 
 	void ShowContextMenu( const FVector2D& ScreenSpacePosition );
 
diff --git a/Engine/Source/Developer/Profiler/Private/Widgets/SProfilerThreadView.h b/Engine/Source/Developer/Profiler/Private/Widgets/SProfilerThreadView.h
--- a/Engine/Source/Developer/Profiler/Private/Widgets/SProfilerThreadView.h
+++ b/Engine/Source/Developer/Profiler/Private/Widgets/SProfilerThreadView.h
@@ -131,7 +131,7 @@ public:
 	* @param  InCurrentTime  Current absolute real time
 	* @param  InDeltaTime  Real time passed since last tick
 	*/
-	virtual void Tick( const FGeometry& AllottedGeometry, const double InCurrentTime, const float InDeltaTime ) OVERRIDE;
+	virtual void Tick( const FGeometry& AllottedGeometry, const double InCurrentTime, const float InDeltaTime ) override;
 
 	/**
 	* The widget should respond by populating the OutDrawElements array with FDrawElements
@@ -146,7 +146,7 @@ public:
 	*
 	* @return The maximum layer ID attained by this widget or any of its children.
 	*/
-	virtual int32 OnPaint( const FGeometry& AllottedGeometry, const FSlateRect& MyClippingRect, FSlateWindowElementList& OutDrawElements, int32 LayerId, const FWidgetStyle& InWidgetStyle, bool bParentEnabled ) const OVERRIDE;
+	virtual int32 OnPaint( const FGeometry& AllottedGeometry, const FSlateRect& MyClippingRect, FSlateWindowElementList& OutDrawElements, int32 LayerId, const FWidgetStyle& InWidgetStyle, bool bParentEnabled ) const override;
 
 	/**
 	* The system calls this method to notify the widget that a mouse button was pressed within it. This event is bubbled.
@@ -156,7 +156,7 @@ public:
 	*
 	* @return Whether the event was handled along with possible requests for the system to take action.
 	*/
-	virtual FReply OnMouseButtonDown( const FGeometry& MyGeometry, const FPointerEvent& MouseEvent ) OVERRIDE;
+	virtual FReply OnMouseButtonDown( const FGeometry& MyGeometry, const FPointerEvent& MouseEvent ) override;
 
 	/**
 	* The system calls this method to notify the widget that a mouse button was release within it. This event is bubbled.
@@ -166,7 +166,7 @@ public:
 	*
 	* @return Whether the event was handled along with possible requests for the system to take action.
 	*/
-	virtual FReply OnMouseButtonUp( const FGeometry& MyGeometry, const FPointerEvent& MouseEvent ) OVERRIDE;
+	virtual FReply OnMouseButtonUp( const FGeometry& MyGeometry, const FPointerEvent& MouseEvent ) override;
 
 	/**
 	* The system calls this method to notify the widget that a mouse moved within it. This event is bubbled.
@@ -176,7 +176,7 @@ public:
 	*
 	* @return Whether the event was handled along with possible requests for the system to take action.
 	*/
-	virtual FReply OnMouseMove( const FGeometry& MyGeometry, const FPointerEvent& MouseEvent ) OVERRIDE;
+	virtual FReply OnMouseMove( const FGeometry& MyGeometry, const FPointerEvent& MouseEvent ) override;
 
 	/**
 	* The system will use this event to notify a widget that the cursor has entered it. This event is NOT bubbled.
@@ -184,14 +184,14 @@ public:
 	* @param MyGeometry The Geometry of the widget receiving the event
 	* @param MouseEvent Information about the input event
 	*/
-	virtual void OnMouseEnter( const FGeometry& MyGeometry, const FPointerEvent& MouseEvent ) OVERRIDE;
+	virtual void OnMouseEnter( const FGeometry& MyGeometry, const FPointerEvent& MouseEvent ) override;
 
 	/**
 	* The system will use this event to notify a widget that the cursor has left it. This event is NOT bubbled.
 	*
 	* @param MouseEvent Information about the input event
 	*/
-	virtual void OnMouseLeave( const FPointerEvent& MouseEvent ) OVERRIDE;
+	virtual void OnMouseLeave( const FPointerEvent& MouseEvent ) override;
 
 	/**
 	* Called when the mouse wheel is spun. This event is bubbled.
@@ -200,7 +200,7 @@ public:
 	*
 	* @return  Returns whether the event was handled, along with other possible actions
 	*/
-	virtual FReply OnMouseWheel( const FGeometry& MyGeometry, const FPointerEvent& MouseEvent ) OVERRIDE;
+	virtual FReply OnMouseWheel( const FGeometry& MyGeometry, const FPointerEvent& MouseEvent ) override;
 
 	/**
 	* Called when a mouse button is double clicked.  Override this in derived classes.
@@ -210,14 +210,14 @@ public:
 	*
 	* @return  Returns whether the event was handled, along with other possible actions
 	*/
-	virtual FReply OnMouseButtonDoubleClick( const FGeometry& MyGeometry, const FPointerEvent& MouseEvent ) OVERRIDE;
+	virtual FReply OnMouseButtonDoubleClick( const FGeometry& MyGeometry, const FPointerEvent& MouseEvent ) override;
 
 	/**
 	* Called when the system wants to know which cursor to display for this Widget.  This event is bubbled.
 	*
 	* @return  The cursor requested (can be None.)
 	*/
-	virtual FCursorReply OnCursorQuery( const FGeometry& MyGeometry, const FPointerEvent& CursorEvent ) const OVERRIDE;
+	virtual FCursorReply OnCursorQuery( const FGeometry& MyGeometry, const FPointerEvent& CursorEvent ) const override;
 
 	void ShowContextMenu( const FVector2D& ScreenSpacePosition );
 
diff --git a/Engine/Source/Developer/Profiler/Private/Widgets/SProfilerWindow.h b/Engine/Source/Developer/Profiler/Private/Widgets/SProfilerWindow.h
--- a/Engine/Source/Developer/Profiler/Private/Widgets/SProfilerWindow.h
+++ b/Engine/Source/Developer/Profiler/Private/Widgets/SProfilerWindow.h
@@ -69,7 +69,7 @@ private:
 	 *
 	 * @return  Returns whether the event was handled, along with other possible actions
 	 */
-	virtual FReply OnKeyDown( const FGeometry& MyGeometry, const FKeyboardEvent& InKeyboardEvent ) OVERRIDE;
+	virtual FReply OnKeyDown( const FGeometry& MyGeometry, const FKeyboardEvent& InKeyboardEvent ) override;
 
 	/**
 	 * Called when the user is dropping something onto a widget; terminates drag and drop.
@@ -79,7 +79,7 @@ private:
 	 *
 	 * @return A reply that indicated whether this event was handled.
 	 */
-	virtual FReply OnDrop( const FGeometry& MyGeometry, const FDragDropEvent& DragDropEvent ) OVERRIDE;
+	virtual FReply OnDrop( const FGeometry& MyGeometry, const FDragDropEvent& DragDropEvent ) override;
 
 	/**
 	 * Called during drag and drop when the the mouse is being dragged over a widget.
@@ -89,7 +89,7 @@ private:
 	 *
 	 * @return A reply that indicated whether this event was handled.
 	 */
-	virtual FReply OnDragOver( const FGeometry& MyGeometry, const FDragDropEvent& DragDropEvent )  OVERRIDE;
+	virtual FReply OnDragOver( const FGeometry& MyGeometry, const FDragDropEvent& DragDropEvent )  override;
 
 	/** Holds all widgets for the profiler window like menu bar, toolbar and tabs. */
 	TSharedPtr<SVerticalBox> MainContentPanel;
diff --git a/Engine/Source/Developer/Profiler/Private/Widgets/StatDragDropOp.h b/Engine/Source/Developer/Profiler/Private/Widgets/StatDragDropOp.h
--- a/Engine/Source/Developer/Profiler/Private/Widgets/StatDragDropOp.h
+++ b/Engine/Source/Developer/Profiler/Private/Widgets/StatDragDropOp.h
@@ -57,7 +57,7 @@ public:
 
 private:
 	/** Gets the widget that will serve as the decorator unless overridden. If you do not override, you will have no decorator. */
-	virtual TSharedPtr<SWidget> GetDefaultDecorator() const OVERRIDE
+	virtual TSharedPtr<SWidget> GetDefaultDecorator() const override
 	{
 		return 
 			
diff --git a/Engine/Source/Developer/ProfilerClient/Private/ProfilerClientManager.h b/Engine/Source/Developer/ProfilerClient/Private/ProfilerClientManager.h
--- a/Engine/Source/Developer/ProfilerClient/Private/ProfilerClientManager.h
+++ b/Engine/Source/Developer/ProfilerClient/Private/ProfilerClientManager.h
@@ -156,68 +156,68 @@ public:
 public:
 
 	// Begin IProfileClient Interface
-	virtual void Subscribe( const FGuid& Session ) OVERRIDE;
+	virtual void Subscribe( const FGuid& Session ) override;
 
-	virtual void Track( const FGuid& Instance ) OVERRIDE;
+	virtual void Track( const FGuid& Instance ) override;
 
-	virtual void Track( const TArray<ISessionInstanceInfoPtr>& Instances ) OVERRIDE;
+	virtual void Track( const TArray<ISessionInstanceInfoPtr>& Instances ) override;
 
-	virtual void Untrack( const FGuid& Instance ) OVERRIDE;
+	virtual void Untrack( const FGuid& Instance ) override;
 
-	virtual void Unsubscribe() OVERRIDE;
+	virtual void Unsubscribe() override;
 
-	virtual void SetCaptureState( const bool bRequestedCaptureState, const FGuid& InstanceId = FGuid() ) OVERRIDE;
+	virtual void SetCaptureState( const bool bRequestedCaptureState, const FGuid& InstanceId = FGuid() ) override;
 
-	virtual void SetPreviewState( const bool bRequestedPreviewState, const FGuid& InstanceId = FGuid() ) OVERRIDE;
+	virtual void SetPreviewState( const bool bRequestedPreviewState, const FGuid& InstanceId = FGuid() ) override;
 
-	virtual void LoadCapture( const FString& DataFilepath, const FGuid& ProfileId ) OVERRIDE;
+	virtual void LoadCapture( const FString& DataFilepath, const FGuid& ProfileId ) override;
 
-	virtual void RequestMetaData() OVERRIDE;
+	virtual void RequestMetaData() override;
 
-	virtual void RequestLastCapturedFile( const FGuid& InstanceId = FGuid() ) OVERRIDE;
+	virtual void RequestLastCapturedFile( const FGuid& InstanceId = FGuid() ) override;
 
-	virtual const FStatMetaData& GetStatMetaData( const FGuid& InstanceId ) const OVERRIDE
+	virtual const FStatMetaData& GetStatMetaData( const FGuid& InstanceId ) const override
 	{
 		return Connections.Find(InstanceId)->MetaData;
 	}
 
-	virtual FProfilerClientDataDelegate& OnProfilerData() OVERRIDE
+	virtual FProfilerClientDataDelegate& OnProfilerData() override
 	{
 		return ProfilerDataDelegate;
 	}
 
-	virtual FProfilerFileTransferDelegate& OnProfilerFileTransfer() OVERRIDE
+	virtual FProfilerFileTransferDelegate& OnProfilerFileTransfer() override
 	{
 		return ProfilerFileTransferDelegate;
 	}
 
-	virtual FProfilerClientConnectedDelegate& OnProfilerClientConnected() OVERRIDE
+	virtual FProfilerClientConnectedDelegate& OnProfilerClientConnected() override
 	{
 		return ProfilerClientConnectedDelegate;
 	}
 
-	virtual FProfilerClientDisconnectedDelegate& OnProfilerClientDisconnected() OVERRIDE
+	virtual FProfilerClientDisconnectedDelegate& OnProfilerClientDisconnected() override
 	{
 		return ProfilerClientDisconnectedDelegate;
 	}
 
-	virtual FProfilerMetaDataUpdateDelegate& OnMetaDataUpdated() OVERRIDE
+	virtual FProfilerMetaDataUpdateDelegate& OnMetaDataUpdated() override
 	{
 		return ProfilerMetaDataUpdatedDelegate;
 	}
 
-	virtual FProfilerLoadStartedDelegate& OnLoadStarted() OVERRIDE
+	virtual FProfilerLoadStartedDelegate& OnLoadStarted() override
 	{
 		return ProfilerLoadStartedDelegate;
 	}
 
-	virtual FProfilerLoadCompletedDelegate& OnLoadCompleted() OVERRIDE
+	virtual FProfilerLoadCompletedDelegate& OnLoadCompleted() override
 	{
 
 		return ProfilerLoadCompletedDelegate;
 	}
 
-	virtual FProfilerLoadedMetaDataDelegate& OnLoadedMetaData() OVERRIDE
+	virtual FProfilerLoadedMetaDataDelegate& OnLoadedMetaData() override
 	{
 		return ProfilerLoadedMetaDataDelegate;
 	}
diff --git a/Engine/Source/Developer/ProfilerClient/Private/ProfilerClientModule.cpp b/Engine/Source/Developer/ProfilerClient/Private/ProfilerClientModule.cpp
--- a/Engine/Source/Developer/ProfilerClient/Private/ProfilerClientModule.cpp
+++ b/Engine/Source/Developer/ProfilerClient/Private/ProfilerClientModule.cpp
@@ -17,12 +17,12 @@ public:
 
 	// Begin IModuleInterface interface
 
-	virtual void StartupModule() OVERRIDE
+	virtual void StartupModule() override
 	{
 		MessageBusPtr = IMessagingModule::Get().GetDefaultBus();
 	}
 
-	virtual void ShutdownModule() OVERRIDE
+	virtual void ShutdownModule() override
 	{
 	}
 
@@ -32,7 +32,7 @@ public:
 
 	// Begin IProfilerClientModule interface
 
-	virtual IProfilerClientPtr CreateProfilerClient() OVERRIDE
+	virtual IProfilerClientPtr CreateProfilerClient() override
 	{
 		IMessageBusPtr MessageBus = MessageBusPtr.Pin();
 
diff --git a/Engine/Source/Developer/ProfilerMessages/Private/ProfilerMessagesModule.cpp b/Engine/Source/Developer/ProfilerMessages/Private/ProfilerMessagesModule.cpp
--- a/Engine/Source/Developer/ProfilerMessages/Private/ProfilerMessagesModule.cpp
+++ b/Engine/Source/Developer/ProfilerMessages/Private/ProfilerMessagesModule.cpp
@@ -16,11 +16,11 @@ class FProfilerMessagesModule
 {
 public:
 
-	virtual void StartupModule( ) OVERRIDE { }
+	virtual void StartupModule( ) override { }
 
-	virtual void ShutdownModule( ) OVERRIDE { }
+	virtual void ShutdownModule( ) override { }
 
-	virtual bool SupportsDynamicReloading( ) OVERRIDE
+	virtual bool SupportsDynamicReloading( ) override
 	{
 		return true;
 	}
diff --git a/Engine/Source/Developer/ProfilerService/Private/ProfilerServiceManager.h b/Engine/Source/Developer/ProfilerService/Private/ProfilerServiceManager.h
--- a/Engine/Source/Developer/ProfilerService/Private/ProfilerServiceManager.h
+++ b/Engine/Source/Developer/ProfilerService/Private/ProfilerServiceManager.h
@@ -50,15 +50,15 @@ public:
 public:
 
 	// Begin IProfilerServiceManager interface
-	virtual void SendData(FProfilerCycleCounter& Data) OVERRIDE;
+	virtual void SendData(FProfilerCycleCounter& Data) override;
 
-	virtual void SendData(FProfilerFloatAccumulator& Data) OVERRIDE;
+	virtual void SendData(FProfilerFloatAccumulator& Data) override;
 
-	virtual void SendData(FProfilerCountAccumulator& Data) OVERRIDE;
+	virtual void SendData(FProfilerCountAccumulator& Data) override;
 
-	virtual void SendData(FProfilerCycleGraph& Data) OVERRIDE;
+	virtual void SendData(FProfilerCycleGraph& Data) override;
 
-	virtual bool IsCapturing() const OVERRIDE
+	virtual bool IsCapturing() const override
 	{ 
 #if STATS
 		return (PreviewClients.Num() > 0 || Archive.IsValid());
@@ -71,16 +71,16 @@ public:
 	
 	virtual void StopCapture();
 
-	virtual void UpdateMetaData() OVERRIDE;
+	virtual void UpdateMetaData() override;
 
-	virtual void StartFrame(uint32 FrameNumber, double FrameStart) OVERRIDE;
+	virtual void StartFrame(uint32 FrameNumber, double FrameStart) override;
 
-	virtual FStatMetaData& GetStatMetaData() OVERRIDE
+	virtual FStatMetaData& GetStatMetaData() override
 	{
 		return MetaData;
 	}
 
-	virtual FProfilerDataDelegate& OnProfilerData() OVERRIDE
+	virtual FProfilerDataDelegate& OnProfilerData() override
 	{
 		return ProfilerDataDelegate;
 	}
diff --git a/Engine/Source/Developer/ProfilerService/Private/ProfilerServiceModule.cpp b/Engine/Source/Developer/ProfilerService/Private/ProfilerServiceModule.cpp
--- a/Engine/Source/Developer/ProfilerService/Private/ProfilerServiceModule.cpp
+++ b/Engine/Source/Developer/ProfilerService/Private/ProfilerServiceModule.cpp
@@ -28,7 +28,7 @@ public:
 	// End IModuleInterface interface
 
 	// Begin IProfilerServiceModule interface
-	virtual IProfilerServiceManagerPtr CreateProfilerServiceManager( ) OVERRIDE
+	virtual IProfilerServiceManagerPtr CreateProfilerServiceManager( ) override
 	{
 		if (!ProfilerServiceManager.IsValid())
 		{
diff --git a/Engine/Source/Developer/RealtimeProfiler/Private/RealtimeProfiler.cpp b/Engine/Source/Developer/RealtimeProfiler/Private/RealtimeProfiler.cpp
--- a/Engine/Source/Developer/RealtimeProfiler/Private/RealtimeProfiler.cpp
+++ b/Engine/Source/Developer/RealtimeProfiler/Private/RealtimeProfiler.cpp
@@ -41,7 +41,7 @@ void FRealtimeProfiler::Update(TSharedPtr< FVisualizerEvent > InProfileData, FRe
 class FRealtimeProfilerModule : public IModuleInterface
 {
 public:
-	virtual void ShutdownModule() OVERRIDE
+	virtual void ShutdownModule() override
 	{
 	}
 };
diff --git a/Engine/Source/Developer/RealtimeProfiler/Private/SRealtimeProfilerLineGraph.h b/Engine/Source/Developer/RealtimeProfiler/Private/SRealtimeProfilerLineGraph.h
--- a/Engine/Source/Developer/RealtimeProfiler/Private/SRealtimeProfilerLineGraph.h
+++ b/Engine/Source/Developer/RealtimeProfiler/Private/SRealtimeProfilerLineGraph.h
@@ -94,13 +94,13 @@ private:
 	void DisplayFrameDetailAtMouse(const FGeometry& InMyGeometry, const FPointerEvent& InMouseEvent);
 
 	// SWidget overrides
-	virtual void Tick( const FGeometry& AllottedGeometry, const double InCurrentTime, const float InDeltaTime ) OVERRIDE;
-	virtual FVector2D ComputeDesiredSize() const OVERRIDE;
-	virtual int32 OnPaint( const FGeometry& AllottedGeometry, const FSlateRect& MyClippingRect, FSlateWindowElementList& OutDrawElements, int32 LayerId, const FWidgetStyle& InWidgetStyle, bool bParentEnabled ) const OVERRIDE;
-	virtual FReply OnMouseButtonDown( const FGeometry& MyGeometry, const FPointerEvent& MouseEvent ) OVERRIDE;
-	virtual FReply OnMouseButtonUp( const FGeometry& MyGeometry, const FPointerEvent& MouseEvent ) OVERRIDE;
-	virtual FReply OnMouseMove( const FGeometry& MyGeometry, const FPointerEvent& MouseEvent ) OVERRIDE;
-	virtual bool SupportsKeyboardFocus() const OVERRIDE
+	virtual void Tick( const FGeometry& AllottedGeometry, const double InCurrentTime, const float InDeltaTime ) override;
+	virtual FVector2D ComputeDesiredSize() const override;
+	virtual int32 OnPaint( const FGeometry& AllottedGeometry, const FSlateRect& MyClippingRect, FSlateWindowElementList& OutDrawElements, int32 LayerId, const FWidgetStyle& InWidgetStyle, bool bParentEnabled ) const override;
+	virtual FReply OnMouseButtonDown( const FGeometry& MyGeometry, const FPointerEvent& MouseEvent ) override;
+	virtual FReply OnMouseButtonUp( const FGeometry& MyGeometry, const FPointerEvent& MouseEvent ) override;
+	virtual FReply OnMouseMove( const FGeometry& MyGeometry, const FPointerEvent& MouseEvent ) override;
+	virtual bool SupportsKeyboardFocus() const override
 	{
 		return true;
 	}
diff --git a/Engine/Source/Developer/ScreenShotComparison/Private/Models/ScreenShotComparisonFilter.h b/Engine/Source/Developer/ScreenShotComparison/Private/Models/ScreenShotComparisonFilter.h
--- a/Engine/Source/Developer/ScreenShotComparison/Private/Models/ScreenShotComparisonFilter.h
+++ b/Engine/Source/Developer/ScreenShotComparison/Private/Models/ScreenShotComparisonFilter.h
@@ -13,7 +13,7 @@ public:
 	// Begin IFilter functions
 
 	DECLARE_DERIVED_EVENT(FScreenShotComparisonFilter, IFilter< const TSharedPtr< class IScreenShotData >& >::FChangedEvent, FChangedEvent);
-	virtual FChangedEvent& OnChanged() OVERRIDE { return ChangedEvent; }
+	virtual FChangedEvent& OnChanged() override { return ChangedEvent; }
 
 
 	/**
@@ -22,7 +22,7 @@ public:
 	 * @param InReport - The automation report
 	 * @return true if it passes the test
 	 */
-	virtual bool PassesFilter( const TSharedPtr< IScreenShotData >& InReport ) const OVERRIDE
+	virtual bool PassesFilter( const TSharedPtr< IScreenShotData >& InReport ) const override
 	{
 		bool FilterPassed = true;
 
diff --git a/Engine/Source/Developer/ScreenShotComparison/Private/ScreenShotComparisonModule.cpp b/Engine/Source/Developer/ScreenShotComparison/Private/ScreenShotComparisonModule.cpp
--- a/Engine/Source/Developer/ScreenShotComparison/Private/ScreenShotComparisonModule.cpp
+++ b/Engine/Source/Developer/ScreenShotComparison/Private/ScreenShotComparisonModule.cpp
@@ -17,7 +17,7 @@ class FScreenShotComparisonModule
 	: public IScreenShotComparisonModule
 {
 public:
-	virtual TSharedRef<class SWidget> CreateScreenShotComparison( const IScreenShotManagerRef& ScreenShotManager ) OVERRIDE
+	virtual TSharedRef<class SWidget> CreateScreenShotComparison( const IScreenShotManagerRef& ScreenShotManager ) override
 	{
 		return SNew(SScreenShotBrowser, ScreenShotManager);
 	}
diff --git a/Engine/Source/Developer/ScreenShotComparisonTools/Private/ScreenShotBaseNode.h b/Engine/Source/Developer/ScreenShotComparisonTools/Private/ScreenShotBaseNode.h
--- a/Engine/Source/Developer/ScreenShotComparisonTools/Private/ScreenShotBaseNode.h
+++ b/Engine/Source/Developer/ScreenShotComparisonTools/Private/ScreenShotBaseNode.h
@@ -30,14 +30,14 @@ public:
 
 	// Begin IScreenShotData interface
 
-	virtual void AddScreenShotData( const FScreenShotDataItem& InScreenDataItem ) OVERRIDE;
-	virtual const FString& GetAssetName() const OVERRIDE;
-	virtual TArray< TSharedPtr< IScreenShotData> >& GetChildren() OVERRIDE;
-	virtual TArray< TSharedPtr< IScreenShotData> >& GetFilteredChildren() OVERRIDE;
-	virtual const FString& GetName() const OVERRIDE;
-	virtual EScreenShotDataType::Type GetScreenNodeType() OVERRIDE;
-	virtual bool SetFilter( TSharedPtr< ScreenShotFilterCollection > ScreenFilter ) OVERRIDE;
-	virtual void SetDisplayEveryNthScreenshot( int32 NewNthValue ) OVERRIDE;
+	virtual void AddScreenShotData( const FScreenShotDataItem& InScreenDataItem ) override;
+	virtual const FString& GetAssetName() const override;
+	virtual TArray< TSharedPtr< IScreenShotData> >& GetChildren() override;
+	virtual TArray< TSharedPtr< IScreenShotData> >& GetFilteredChildren() override;
+	virtual const FString& GetName() const override;
+	virtual EScreenShotDataType::Type GetScreenNodeType() override;
+	virtual bool SetFilter( TSharedPtr< ScreenShotFilterCollection > ScreenFilter ) override;
+	virtual void SetDisplayEveryNthScreenshot( int32 NewNthValue ) override;
 
 	// End IScreenShotData interface
 
diff --git a/Engine/Source/Developer/ScreenShotComparisonTools/Private/ScreenShotManager.h b/Engine/Source/Developer/ScreenShotComparisonTools/Private/ScreenShotManager.h
--- a/Engine/Source/Developer/ScreenShotComparisonTools/Private/ScreenShotManager.h
+++ b/Engine/Source/Developer/ScreenShotComparisonTools/Private/ScreenShotManager.h
@@ -33,17 +33,17 @@ public:
 
 	// Begin IScreenShotManager interface
 
-	virtual void GenerateLists() OVERRIDE;
+	virtual void GenerateLists() override;
 
-	virtual TArray< TSharedPtr<FString> >& GetCachedPlatfomList(  ) OVERRIDE;
+	virtual TArray< TSharedPtr<FString> >& GetCachedPlatfomList(  ) override;
 
-	virtual TArray<IScreenShotDataPtr>& GetLists() OVERRIDE;
+	virtual TArray<IScreenShotDataPtr>& GetLists() override;
 
-	virtual void RegisterScreenShotUpdate(const FOnScreenFilterChanged& InDelegate )OVERRIDE;
+	virtual void RegisterScreenShotUpdate(const FOnScreenFilterChanged& InDelegate )override;
 
-	virtual void SetFilter( TSharedPtr< ScreenShotFilterCollection > InFilter ) OVERRIDE;
+	virtual void SetFilter( TSharedPtr< ScreenShotFilterCollection > InFilter ) override;
 
-	virtual void SetDisplayEveryNthScreenshot(int32 NewNth) OVERRIDE;
+	virtual void SetDisplayEveryNthScreenshot(int32 NewNth) override;
 
 	// End IScreenShotManager interface
 
diff --git a/Engine/Source/Developer/ScreenShotComparisonTools/Private/ScreenShotPlatformNode.h b/Engine/Source/Developer/ScreenShotComparisonTools/Private/ScreenShotPlatformNode.h
--- a/Engine/Source/Developer/ScreenShotComparisonTools/Private/ScreenShotPlatformNode.h
+++ b/Engine/Source/Developer/ScreenShotComparisonTools/Private/ScreenShotPlatformNode.h
@@ -25,7 +25,7 @@ public:
 
 	// Begin IScreenShotData interface
 
-	virtual void AddScreenShotData( const FScreenShotDataItem& InScreenDataItem ) OVERRIDE 
+	virtual void AddScreenShotData( const FScreenShotDataItem& InScreenDataItem ) override 
 	{
 		FString ScreenShotNumber = FString::Printf( TEXT( "CL #%d"), InScreenDataItem.ChangeListNumber );
 		TSharedRef<FScreenShotBaseNode> ChildNode = MakeShareable( new FScreenShotBaseNode( ScreenShotNumber, InScreenDataItem.AssetName) );
@@ -33,13 +33,13 @@ public:
 	};
 
 
-	virtual EScreenShotDataType::Type GetScreenNodeType() OVERRIDE 
+	virtual EScreenShotDataType::Type GetScreenNodeType() override 
 	{ 
 		return EScreenShotDataType::SSDT_Platform; 
 	};
 
 
-	virtual bool SetFilter( TSharedPtr< ScreenShotFilterCollection > ScreenFilter ) OVERRIDE 
+	virtual bool SetFilter( TSharedPtr< ScreenShotFilterCollection > ScreenFilter ) override 
 	{
 		FilteredChildren.Empty();
 
@@ -59,7 +59,7 @@ public:
 		return bPassesFilter;
 	};
 
-	virtual void SetDisplayEveryNthScreenshot( int32 NewLastNth ) OVERRIDE
+	virtual void SetDisplayEveryNthScreenshot( int32 NewLastNth ) override
 	{
 		DisplayEveryNthScreenshot = NewLastNth;
 		if( FilteredChildren.Num() > 0 )
diff --git a/Engine/Source/Developer/ScreenShotComparisonTools/Private/ScreenShotScreenNode.h b/Engine/Source/Developer/ScreenShotComparisonTools/Private/ScreenShotScreenNode.h
--- a/Engine/Source/Developer/ScreenShotComparisonTools/Private/ScreenShotScreenNode.h
+++ b/Engine/Source/Developer/ScreenShotComparisonTools/Private/ScreenShotScreenNode.h
@@ -24,7 +24,7 @@ public:
 
 	// Begin IScreenShotData interface
 
-	virtual void AddScreenShotData( const FScreenShotDataItem& InScreenDataItem ) OVERRIDE 
+	virtual void AddScreenShotData( const FScreenShotDataItem& InScreenDataItem ) override 
 	{
 		// First stage of adding screen shot data - create the screen shot node, then add the platform
 		TSharedRef<IScreenShotData> ScreenNode = AddChild( InScreenDataItem.DeviceName );
@@ -32,12 +32,12 @@ public:
 	};
 
 
-	virtual EScreenShotDataType::Type GetScreenNodeType() OVERRIDE 
+	virtual EScreenShotDataType::Type GetScreenNodeType() override 
 	{ 
 		return EScreenShotDataType::SSDT_ScreenView; 
 	};
 
-	virtual bool SetFilter( TSharedPtr< ScreenShotFilterCollection > ScreenFilter ) OVERRIDE
+	virtual bool SetFilter( TSharedPtr< ScreenShotFilterCollection > ScreenFilter ) override
 	{
 		FilteredChildren.Empty();
 
@@ -61,7 +61,7 @@ public:
 
 	// Begin SScreenShotNodeBase interface
 
-	virtual TSharedRef<IScreenShotData> CreateNode( const FString& ChildName ) OVERRIDE 
+	virtual TSharedRef<IScreenShotData> CreateNode( const FString& ChildName ) override 
 	{
 		TSharedPtr< IScreenShotData > ScreenShotNode = MakeShareable( new FScreenShotPlatformNode( ChildName ) );;
 		Children.Add( ScreenShotNode.ToSharedRef() );
diff --git a/Engine/Source/Developer/ScreenShotComparisonTools/Private/ScreenShotToolsModule.cpp b/Engine/Source/Developer/ScreenShotComparisonTools/Private/ScreenShotToolsModule.cpp
--- a/Engine/Source/Developer/ScreenShotComparisonTools/Private/ScreenShotToolsModule.cpp
+++ b/Engine/Source/Developer/ScreenShotComparisonTools/Private/ScreenShotToolsModule.cpp
@@ -19,12 +19,12 @@ public:
 
 	// Begin IScreenShotToolsModule interface
 
-	virtual IScreenShotManagerPtr GetScreenShotManager( ) OVERRIDE
+	virtual IScreenShotManagerPtr GetScreenShotManager( ) override
 	{
 		return ScreenShotManager;
 	}
 
-	virtual void UpdateScreenShotData( ) OVERRIDE
+	virtual void UpdateScreenShotData( ) override
 	{
 		if (ScreenShotManager.IsValid())
 		{
@@ -40,7 +40,7 @@ public:
 	/*
 	 * Shutdown the screen shot manager tools module
 	 */
-	virtual void ShutdownModule( ) OVERRIDE
+	virtual void ShutdownModule( ) override
 	{
 		ScreenShotManager.Reset();
 	}
@@ -48,7 +48,7 @@ public:
 	/*
 	* Startup the screen shot manager tools module
 	*/
-	void StartupModule() OVERRIDE
+	void StartupModule() override
 	{
 		IMessageBusPtr MessageBus = IMessagingModule::Get().GetDefaultBus();
 		check(MessageBus.IsValid());
diff --git a/Engine/Source/Developer/SessionFrontend/Private/Widgets/Console/SSessionConsole.h b/Engine/Source/Developer/SessionFrontend/Private/Widgets/Console/SSessionConsole.h
--- a/Engine/Source/Developer/SessionFrontend/Private/Widgets/Console/SSessionConsole.h
+++ b/Engine/Source/Developer/SessionFrontend/Private/Widgets/Console/SSessionConsole.h
@@ -72,7 +72,7 @@ protected:
 
 protected:
 
-	virtual FReply OnKeyDown( const FGeometry& MyGeometry, const FKeyboardEvent& InKeyboardEvent ) OVERRIDE;
+	virtual FReply OnKeyDown( const FGeometry& MyGeometry, const FKeyboardEvent& InKeyboardEvent ) override;
 
 private:
 
diff --git a/Engine/Source/Developer/SessionLauncher/Private/Widgets/SSessionLauncherProgress.h b/Engine/Source/Developer/SessionLauncher/Private/Widgets/SSessionLauncherProgress.h
--- a/Engine/Source/Developer/SessionLauncher/Private/Widgets/SSessionLauncherProgress.h
+++ b/Engine/Source/Developer/SessionLauncher/Private/Widgets/SSessionLauncherProgress.h
@@ -188,7 +188,7 @@ public:
 		MessageListView->RequestListRefresh();
 	}
 
-	void Tick( const FGeometry& AllottedGeometry, const double InCurrentTime, const float InDeltaTime ) OVERRIDE
+	void Tick( const FGeometry& AllottedGeometry, const double InCurrentTime, const float InDeltaTime ) override
 	{
 		if (PendingMessages.Num() > 0)
 		{
diff --git a/Engine/Source/Developer/SharedSettingsWidgets/Public/SPlatformSetupMessage.h b/Engine/Source/Developer/SharedSettingsWidgets/Public/SPlatformSetupMessage.h
--- a/Engine/Source/Developer/SharedSettingsWidgets/Public/SPlatformSetupMessage.h
+++ b/Engine/Source/Developer/SharedSettingsWidgets/Public/SPlatformSetupMessage.h
@@ -34,7 +34,7 @@ public:
 	TAttribute<bool> GetReadyToGoAttribute() const;
 
 	// SWidget interface
-	virtual void Tick(const FGeometry& AllottedGeometry, const double InCurrentTime, const float InDeltaTime) OVERRIDE;
+	virtual void Tick(const FGeometry& AllottedGeometry, const double InCurrentTime, const float InDeltaTime) override;
 	// End of Swidget interface
 private:
 	int32 GetSetupStateAsInt() const;
diff --git a/Engine/Source/Developer/SimplygonMeshReduction/Private/SimplygonMeshReduction.cpp b/Engine/Source/Developer/SimplygonMeshReduction/Private/SimplygonMeshReduction.cpp
--- a/Engine/Source/Developer/SimplygonMeshReduction/Private/SimplygonMeshReduction.cpp
+++ b/Engine/Source/Developer/SimplygonMeshReduction/Private/SimplygonMeshReduction.cpp
@@ -13,12 +13,12 @@ class FSimplygonMeshReductionModule : public IMeshReductionModule
 {
 public:
 	// IModuleInterface interface.
-	virtual void StartupModule() OVERRIDE;
-	virtual void ShutdownModule() OVERRIDE;
+	virtual void StartupModule() override;
+	virtual void ShutdownModule() override;
 
 	// IMeshReductionModule interface.
-	virtual class IMeshReduction* GetMeshReductionInterface() OVERRIDE;
-	virtual class IMeshMerging* GetMeshMergingInterface() OVERRIDE;
+	virtual class IMeshReduction* GetMeshReductionInterface() override;
+	virtual class IMeshMerging* GetMeshMergingInterface() override;
 };
 
 
@@ -86,7 +86,7 @@ class FSimplygonMeshReduction
 	, public IMeshMerging
 {
 public:
-	virtual const FString& GetVersionString() const OVERRIDE
+	virtual const FString& GetVersionString() const override
 	{
 		return VersionString;
 	}
@@ -406,7 +406,7 @@ public:
 		const TArray<MaterialExportUtils::FFlattenMaterial>& InputMaterials,
 		const struct FMeshProxySettings& InProxySettings,
 		FRawMesh& OutProxyMesh,
-		MaterialExportUtils::FFlattenMaterial& OutMaterial) OVERRIDE
+		MaterialExportUtils::FFlattenMaterial& OutMaterial) override
 	{
 		if (InputMeshes.Num() == 0)
 		{
diff --git a/Engine/Source/Developer/SourceCodeAccess/Private/SourceCodeAccessModule.h b/Engine/Source/Developer/SourceCodeAccess/Private/SourceCodeAccessModule.h
--- a/Engine/Source/Developer/SourceCodeAccess/Private/SourceCodeAccessModule.h
+++ b/Engine/Source/Developer/SourceCodeAccess/Private/SourceCodeAccessModule.h
@@ -14,16 +14,16 @@ public:
 	FSourceCodeAccessModule();
 
 	/** IModuleInterface implementation */
-	virtual void StartupModule() OVERRIDE;
-	virtual void ShutdownModule() OVERRIDE;
+	virtual void StartupModule() override;
+	virtual void ShutdownModule() override;
 
 	/** ISourceCodeAccessModule implementation */
-	virtual bool CanAccessSourceCode() const OVERRIDE;
-	virtual ISourceCodeAccessor& GetAccessor() const OVERRIDE;
-	virtual void SetAccessor(const FName& InName) OVERRIDE;
-	virtual FLaunchingCodeAccessor& OnLaunchingCodeAccessor() OVERRIDE;
-	virtual FDoneLaunchingCodeAccessor& OnDoneLaunchingCodeAccessor() OVERRIDE;
-	virtual FOpenFileFailed& OnOpenFileFailed() OVERRIDE;
+	virtual bool CanAccessSourceCode() const override;
+	virtual ISourceCodeAccessor& GetAccessor() const override;
+	virtual void SetAccessor(const FName& InName) override;
+	virtual FLaunchingCodeAccessor& OnLaunchingCodeAccessor() override;
+	virtual FDoneLaunchingCodeAccessor& OnDoneLaunchingCodeAccessor() override;
+	virtual FOpenFileFailed& OnOpenFileFailed() override;
 
 private:
 	/** Handle when one of the modular features we are interested in is registered */
diff --git a/Engine/Source/Developer/SourceControl/Private/DefaultSourceControlProvider.h b/Engine/Source/Developer/SourceControl/Private/DefaultSourceControlProvider.h
--- a/Engine/Source/Developer/SourceControl/Private/DefaultSourceControlProvider.h
+++ b/Engine/Source/Developer/SourceControl/Private/DefaultSourceControlProvider.h
@@ -11,22 +11,22 @@ class FDefaultSourceControlProvider : public ISourceControlProvider
 {
 public:
 	// ISourceControlProvider implementation
-	virtual void Init(bool bForceConnection = true) OVERRIDE;
-	virtual void Close() OVERRIDE;
-	virtual FText GetStatusText() const OVERRIDE;
-	virtual bool IsAvailable() const OVERRIDE;
-	virtual bool IsEnabled() const OVERRIDE;
-	virtual const FName& GetName(void) const OVERRIDE;
-	virtual ECommandResult::Type GetState( const TArray<FString>& InFiles, TArray< TSharedRef<ISourceControlState, ESPMode::ThreadSafe> >& OutState, EStateCacheUsage::Type InStateCacheUsage ) OVERRIDE;
-	virtual void RegisterSourceControlStateChanged( const FSourceControlStateChanged::FDelegate& SourceControlStateChanged ) OVERRIDE;
-	virtual void UnregisterSourceControlStateChanged( const FSourceControlStateChanged::FDelegate& SourceControlStateChanged ) OVERRIDE;
-	virtual ECommandResult::Type Execute( const TSharedRef<ISourceControlOperation, ESPMode::ThreadSafe>& InOperation, const TArray<FString>& InFiles, EConcurrency::Type InConcurrency = EConcurrency::Synchronous, const FSourceControlOperationComplete& InOperationCompleteDelegate = FSourceControlOperationComplete() ) OVERRIDE;
-	virtual bool CanCancelOperation( const TSharedRef<ISourceControlOperation, ESPMode::ThreadSafe>& InOperation ) const OVERRIDE;
-	virtual void CancelOperation( const TSharedRef<ISourceControlOperation, ESPMode::ThreadSafe>& InOperation ) OVERRIDE;
-	virtual bool UsesLocalReadOnlyState() const OVERRIDE;
-	virtual void Tick() OVERRIDE;
-	virtual TArray< TSharedRef<class ISourceControlLabel> > GetLabels( const FString& InMatchingSpec ) const OVERRIDE;
+	virtual void Init(bool bForceConnection = true) override;
+	virtual void Close() override;
+	virtual FText GetStatusText() const override;
+	virtual bool IsAvailable() const override;
+	virtual bool IsEnabled() const override;
+	virtual const FName& GetName(void) const override;
+	virtual ECommandResult::Type GetState( const TArray<FString>& InFiles, TArray< TSharedRef<ISourceControlState, ESPMode::ThreadSafe> >& OutState, EStateCacheUsage::Type InStateCacheUsage ) override;
+	virtual void RegisterSourceControlStateChanged( const FSourceControlStateChanged::FDelegate& SourceControlStateChanged ) override;
+	virtual void UnregisterSourceControlStateChanged( const FSourceControlStateChanged::FDelegate& SourceControlStateChanged ) override;
+	virtual ECommandResult::Type Execute( const TSharedRef<ISourceControlOperation, ESPMode::ThreadSafe>& InOperation, const TArray<FString>& InFiles, EConcurrency::Type InConcurrency = EConcurrency::Synchronous, const FSourceControlOperationComplete& InOperationCompleteDelegate = FSourceControlOperationComplete() ) override;
+	virtual bool CanCancelOperation( const TSharedRef<ISourceControlOperation, ESPMode::ThreadSafe>& InOperation ) const override;
+	virtual void CancelOperation( const TSharedRef<ISourceControlOperation, ESPMode::ThreadSafe>& InOperation ) override;
+	virtual bool UsesLocalReadOnlyState() const override;
+	virtual void Tick() override;
+	virtual TArray< TSharedRef<class ISourceControlLabel> > GetLabels( const FString& InMatchingSpec ) const override;
 #if SOURCE_CONTROL_WITH_SLATE
-	virtual TSharedRef<class SWidget> MakeSettingsWidget() const OVERRIDE;
+	virtual TSharedRef<class SWidget> MakeSettingsWidget() const override;
 #endif // SOURCE_CONTROL_WITH_SLATE
 };
diff --git a/Engine/Source/Developer/SourceControl/Private/SourceControlModule.h b/Engine/Source/Developer/SourceControl/Private/SourceControlModule.h
--- a/Engine/Source/Developer/SourceControl/Private/SourceControlModule.h
+++ b/Engine/Source/Developer/SourceControl/Private/SourceControlModule.h
@@ -18,22 +18,22 @@ public:
 	FSourceControlModule();
 
 	/** IModuleInterface implementation */
-	virtual void StartupModule() OVERRIDE;
-	virtual void ShutdownModule() OVERRIDE;
+	virtual void StartupModule() override;
+	virtual void ShutdownModule() override;
 
 	/** ISourceControlModule implementation */
-	virtual void Tick() OVERRIDE;
-	virtual void QueueStatusUpdate(const TArray<UPackage*>& InPackages) OVERRIDE;
-	virtual void QueueStatusUpdate(const TArray<FString>& InFilenames) OVERRIDE;
-	virtual void QueueStatusUpdate(UPackage* InPackage) OVERRIDE;
-	virtual void QueueStatusUpdate(const FString& InFilename) OVERRIDE;
-	virtual bool IsEnabled() const OVERRIDE;
-	virtual ISourceControlProvider& GetProvider() const OVERRIDE;
-	virtual void SetProvider( const FName& InName ) OVERRIDE;
-	virtual void ShowLoginDialog(const FSourceControlLoginClosed& InOnSourceControlLoginClosed, ELoginWindowMode::Type InLoginWindowMode, EOnLoginWindowStartup::Type InOnLoginWindowStartup = EOnLoginWindowStartup::ResetProviderToNone) OVERRIDE;
-	virtual TSharedPtr<class SWidget> CreateStatusWidget() const OVERRIDE;
-	virtual bool GetUseGlobalSettings() const OVERRIDE;
-	virtual void SetUseGlobalSettings(bool bIsUseGlobalSettings) OVERRIDE;
+	virtual void Tick() override;
+	virtual void QueueStatusUpdate(const TArray<UPackage*>& InPackages) override;
+	virtual void QueueStatusUpdate(const TArray<FString>& InFilenames) override;
+	virtual void QueueStatusUpdate(UPackage* InPackage) override;
+	virtual void QueueStatusUpdate(const FString& InFilename) override;
+	virtual bool IsEnabled() const override;
+	virtual ISourceControlProvider& GetProvider() const override;
+	virtual void SetProvider( const FName& InName ) override;
+	virtual void ShowLoginDialog(const FSourceControlLoginClosed& InOnSourceControlLoginClosed, ELoginWindowMode::Type InLoginWindowMode, EOnLoginWindowStartup::Type InOnLoginWindowStartup = EOnLoginWindowStartup::ResetProviderToNone) override;
+	virtual TSharedPtr<class SWidget> CreateStatusWidget() const override;
+	virtual bool GetUseGlobalSettings() const override;
+	virtual void SetUseGlobalSettings(bool bIsUseGlobalSettings) override;
 
 	/** Save the settings to the ini file */
 	void SaveSettings();
diff --git a/Engine/Source/Developer/SourceControl/Public/SourceControlOperations.h b/Engine/Source/Developer/SourceControl/Public/SourceControlOperations.h
--- a/Engine/Source/Developer/SourceControl/Public/SourceControlOperations.h
+++ b/Engine/Source/Developer/SourceControl/Public/SourceControlOperations.h
@@ -13,12 +13,12 @@ class FConnect : public ISourceControlOperation
 {
 public:
 	// ISourceControlOperation interface
-	virtual FName GetName() const OVERRIDE 
+	virtual FName GetName() const override 
 	{ 
 		return "Connect"; 
 	}
 
-	virtual FText GetInProgressString() const OVERRIDE
+	virtual FText GetInProgressString() const override
 	{ 
 		return LOCTEXT("SourceControl_Connecting", "Connecting to source control..."); 
 	}
@@ -58,12 +58,12 @@ class FCheckIn : public ISourceControlOperation
 {
 public:
 	// ISourceControlOperation interface
-	virtual FName GetName() const OVERRIDE
+	virtual FName GetName() const override
 	{
 		return "CheckIn";
 	}
 
-	virtual FText GetInProgressString() const OVERRIDE
+	virtual FText GetInProgressString() const override
 	{ 
 		return LOCTEXT("SourceControl_CheckIn", "Checking file(s) into Source Control..."); 
 	}
@@ -103,12 +103,12 @@ class FCheckOut : public ISourceControlOperation
 {
 public:
 	// ISourceControlOperation interface
-	virtual FName GetName() const OVERRIDE
+	virtual FName GetName() const override
 	{
 		return "CheckOut";
 	}
 
-	virtual FText GetInProgressString() const OVERRIDE
+	virtual FText GetInProgressString() const override
 	{ 
 		return LOCTEXT("SourceControl_CheckOut", "Checking file(s) out of Source Control..."); 
 	}
@@ -121,12 +121,12 @@ class FMarkForAdd : public ISourceControlOperation
 {
 public:
 	// ISourceControlOperation interface
-	virtual FName GetName() const OVERRIDE
+	virtual FName GetName() const override
 	{
 		return "MarkForAdd";
 	}
 
-	virtual FText GetInProgressString() const OVERRIDE
+	virtual FText GetInProgressString() const override
 	{ 
 		return LOCTEXT("SourceControl_Add", "Adding file(s) to Source Control..."); 
 	}
@@ -139,12 +139,12 @@ class FDelete : public ISourceControlOperation
 {
 public:
 	// ISourceControlOperation interface
-	virtual FName GetName() const OVERRIDE
+	virtual FName GetName() const override
 	{
 		return "Delete";
 	}
 
-	virtual FText GetInProgressString() const OVERRIDE
+	virtual FText GetInProgressString() const override
 	{ 
 		return LOCTEXT("SourceControl_Delete", "Deleting file(s) from Source Control..."); 
 	}
@@ -157,12 +157,12 @@ class FRevert : public ISourceControlOperation
 {
 public:
 	// ISourceControlOperation interface
-	virtual FName GetName() const OVERRIDE
+	virtual FName GetName() const override
 	{
 		return "Revert";
 	}
 
-	virtual FText GetInProgressString() const OVERRIDE
+	virtual FText GetInProgressString() const override
 	{ 
 		return LOCTEXT("SourceControl_Revert", "Reverting file(s) in Source Control..."); 
 	}	
@@ -175,12 +175,12 @@ class FSync : public ISourceControlOperation
 {
 public:
 	// ISourceControlOperation interface
-	virtual FName GetName() const OVERRIDE
+	virtual FName GetName() const override
 	{
 		return "Sync";
 	}
 
-	virtual FText GetInProgressString() const OVERRIDE
+	virtual FText GetInProgressString() const override
 	{ 
 		return LOCTEXT("SourceControl_Sync", "Syncing file(s) from source control..."); 
 	}	
@@ -209,12 +209,12 @@ public:
 	}
 
 	// ISourceControlOperation interface
-	virtual FName GetName() const OVERRIDE
+	virtual FName GetName() const override
 	{
 		return "UpdateStatus";
 	}
 
-	virtual FText GetInProgressString() const OVERRIDE
+	virtual FText GetInProgressString() const override
 	{ 
 		return LOCTEXT("SourceControl_Update", "Updating file(s) source control status..."); 
 	}	
@@ -267,12 +267,12 @@ class FCopy : public ISourceControlOperation
 {
 public:
 	// ISourceControlOperation interface
-	virtual FName GetName() const OVERRIDE
+	virtual FName GetName() const override
 	{
 		return "Copy";
 	}
 
-	virtual FText GetInProgressString() const OVERRIDE
+	virtual FText GetInProgressString() const override
 	{ 
 		return LOCTEXT("SourceControl_Copy", "Copying file(s) in Source Control..."); 
 	}	
diff --git a/Engine/Source/Developer/StandaloneRenderer/Private/OpenGL/SlateOpenGLRenderer.cpp b/Engine/Source/Developer/StandaloneRenderer/Private/OpenGL/SlateOpenGLRenderer.cpp
--- a/Engine/Source/Developer/StandaloneRenderer/Private/OpenGL/SlateOpenGLRenderer.cpp
+++ b/Engine/Source/Developer/StandaloneRenderer/Private/OpenGL/SlateOpenGLRenderer.cpp
@@ -38,7 +38,7 @@ public:
 	{
 	}
 
-	virtual TSharedRef<FSlateFontAtlas> CreateFontAtlas() const OVERRIDE
+	virtual TSharedRef<FSlateFontAtlas> CreateFontAtlas() const override
 	{
 		/** Size of each font texture, width and height */
 		const uint32 TextureSize = 1024;
diff --git a/Engine/Source/Developer/StandaloneRenderer/Private/OpenGL/SlateOpenGLRenderer.h b/Engine/Source/Developer/StandaloneRenderer/Private/OpenGL/SlateOpenGLRenderer.h
--- a/Engine/Source/Developer/StandaloneRenderer/Private/OpenGL/SlateOpenGLRenderer.h
+++ b/Engine/Source/Developer/StandaloneRenderer/Private/OpenGL/SlateOpenGLRenderer.h
@@ -132,17 +132,17 @@ public:
 	~FSlateOpenGLRenderer();
 
 	/** FSlateRenderer interface */
-	virtual FSlateDrawBuffer& GetDrawBuffer() OVERRIDE;
-	virtual void Initialize() OVERRIDE;
-	virtual void Destroy() OVERRIDE {}
-	virtual void DrawWindows( FSlateDrawBuffer& InWindowDrawBuffer ) OVERRIDE;
-	virtual void OnWindowDestroyed( const TSharedRef<SWindow>& InWindow ) OVERRIDE;
-	virtual void CreateViewport( const TSharedRef<SWindow> InWindow ) OVERRIDE;
-	virtual void RequestResize( const TSharedPtr<SWindow>& InWindow, uint32 NewSizeX, uint32 NewSizeY ) OVERRIDE;
-	virtual void UpdateFullscreenState( const TSharedRef<SWindow> InWindow, uint32 OverrideResX, uint32 OverrideResY ) OVERRIDE;
-	virtual void ReleaseDynamicResource( const FSlateBrush& Brush ) OVERRIDE;
-	virtual bool GenerateDynamicImageResource(FName ResourceName, uint32 Width, uint32 Height, const TArray< uint8 >& Bytes) OVERRIDE;
-	virtual void LoadStyleResources(const ISlateStyle& Style) OVERRIDE;
+	virtual FSlateDrawBuffer& GetDrawBuffer() override;
+	virtual void Initialize() override;
+	virtual void Destroy() override {}
+	virtual void DrawWindows( FSlateDrawBuffer& InWindowDrawBuffer ) override;
+	virtual void OnWindowDestroyed( const TSharedRef<SWindow>& InWindow ) override;
+	virtual void CreateViewport( const TSharedRef<SWindow> InWindow ) override;
+	virtual void RequestResize( const TSharedPtr<SWindow>& InWindow, uint32 NewSizeX, uint32 NewSizeY ) override;
+	virtual void UpdateFullscreenState( const TSharedRef<SWindow> InWindow, uint32 OverrideResX, uint32 OverrideResY ) override;
+	virtual void ReleaseDynamicResource( const FSlateBrush& Brush ) override;
+	virtual bool GenerateDynamicImageResource(FName ResourceName, uint32 Width, uint32 Height, const TArray< uint8 >& Bytes) override;
+	virtual void LoadStyleResources(const ISlateStyle& Style) override;
 
 private:
 	/** 
diff --git a/Engine/Source/Developer/StandaloneRenderer/Private/Windows/D3D/SlateD3DRenderer.cpp b/Engine/Source/Developer/StandaloneRenderer/Private/Windows/D3D/SlateD3DRenderer.cpp
--- a/Engine/Source/Developer/StandaloneRenderer/Private/Windows/D3D/SlateD3DRenderer.cpp
+++ b/Engine/Source/Developer/StandaloneRenderer/Private/Windows/D3D/SlateD3DRenderer.cpp
@@ -52,7 +52,7 @@ public:
 	{
 	}
 
-	virtual TSharedRef<FSlateFontAtlas> CreateFontAtlas() const OVERRIDE
+	virtual TSharedRef<FSlateFontAtlas> CreateFontAtlas() const override
 	{
 		/** Size of each font texture, width and height */
 		const uint32 TextureSize = 1024;
diff --git a/Engine/Source/Developer/StandaloneRenderer/Private/Windows/D3D/SlateD3DRenderer.h b/Engine/Source/Developer/StandaloneRenderer/Private/Windows/D3D/SlateD3DRenderer.h
--- a/Engine/Source/Developer/StandaloneRenderer/Private/Windows/D3D/SlateD3DRenderer.h
+++ b/Engine/Source/Developer/StandaloneRenderer/Private/Windows/D3D/SlateD3DRenderer.h
@@ -44,17 +44,17 @@ public:
 	~FSlateD3DRenderer();
 
 	/** FSlateRenderer Interface */
-	virtual void Initialize() OVERRIDE;
-	virtual void Destroy() OVERRIDE;
-	virtual FSlateDrawBuffer& GetDrawBuffer() OVERRIDE; 
-	virtual void DrawWindows( FSlateDrawBuffer& InWindowDrawBuffer ) OVERRIDE;
-	virtual void OnWindowDestroyed( const TSharedRef<SWindow>& InWindow ) OVERRIDE;
-	virtual void CreateViewport( const TSharedRef<SWindow> InWindow ) OVERRIDE;
-	virtual void RequestResize( const TSharedPtr<SWindow>& InWindow, uint32 NewSizeX, uint32 NewSizeY ) OVERRIDE;
-	virtual void UpdateFullscreenState( const TSharedRef<SWindow> InWindow, uint32 OverrideResX, uint32 OverrideResY ) OVERRIDE;
-	virtual void ReleaseDynamicResource( const FSlateBrush& Brush ) OVERRIDE;
-	virtual bool GenerateDynamicImageResource(FName ResourceName, uint32 Width, uint32 Height, const TArray< uint8 >& Bytes) OVERRIDE;
-	virtual void LoadStyleResources( const ISlateStyle& Style ) OVERRIDE;
+	virtual void Initialize() override;
+	virtual void Destroy() override;
+	virtual FSlateDrawBuffer& GetDrawBuffer() override; 
+	virtual void DrawWindows( FSlateDrawBuffer& InWindowDrawBuffer ) override;
+	virtual void OnWindowDestroyed( const TSharedRef<SWindow>& InWindow ) override;
+	virtual void CreateViewport( const TSharedRef<SWindow> InWindow ) override;
+	virtual void RequestResize( const TSharedPtr<SWindow>& InWindow, uint32 NewSizeX, uint32 NewSizeY ) override;
+	virtual void UpdateFullscreenState( const TSharedRef<SWindow> InWindow, uint32 OverrideResX, uint32 OverrideResY ) override;
+	virtual void ReleaseDynamicResource( const FSlateBrush& Brush ) override;
+	virtual bool GenerateDynamicImageResource(FName ResourceName, uint32 Width, uint32 Height, const TArray< uint8 >& Bytes) override;
+	virtual void LoadStyleResources( const ISlateStyle& Style ) override;
 	
 	void CreateDevice();
 	void CreateDepthStencilBuffer( FSlateD3DViewport& Viewport );
diff --git a/Engine/Source/Developer/StandaloneRenderer/Private/Windows/D3D/SlateD3DTextures.h b/Engine/Source/Developer/StandaloneRenderer/Private/Windows/D3D/SlateD3DTextures.h
--- a/Engine/Source/Developer/StandaloneRenderer/Private/Windows/D3D/SlateD3DTextures.h
+++ b/Engine/Source/Developer/StandaloneRenderer/Private/Windows/D3D/SlateD3DTextures.h
@@ -42,7 +42,7 @@ public:
 	~FSlateTextureAtlasD3D();
 
 	void InitAtlasTexture( int32 Index );
-	virtual void ConditionalUpdateTexture() OVERRIDE;
+	virtual void ConditionalUpdateTexture() override;
 
 	FSlateD3DTexture* GetAtlasTexture() const { return AtlasTexture; }
 private:
diff --git a/Engine/Source/Developer/SynthBenchmark/Private/SynthBenchmarkPrivatePCH.cpp b/Engine/Source/Developer/SynthBenchmark/Private/SynthBenchmarkPrivatePCH.cpp
--- a/Engine/Source/Developer/SynthBenchmark/Private/SynthBenchmarkPrivatePCH.cpp
+++ b/Engine/Source/Developer/SynthBenchmark/Private/SynthBenchmarkPrivatePCH.cpp
@@ -17,11 +17,11 @@ static float GGlobalStateObject = 0.0f;
 class FSynthBenchmark : public ISynthBenchmark
 {
 	/** IModuleInterface implementation */
-	virtual void StartupModule() OVERRIDE;
-	virtual void ShutdownModule() OVERRIDE;
+	virtual void StartupModule() override;
+	virtual void ShutdownModule() override;
 	
 	/** ISynthBenchmark implementation */
-	virtual void Run(FSynthBenchmarkResults& InOut, bool bGPUBenchmark, uint32 WorkScale, bool bDebugOut) const OVERRIDE;
+	virtual void Run(FSynthBenchmarkResults& InOut, bool bGPUBenchmark, uint32 WorkScale, bool bDebugOut) const override;
 };
 
 IMPLEMENT_MODULE( FSynthBenchmark, SynthBenchmark )
diff --git a/Engine/Source/Developer/TargetPlatform/Private/TargetPlatformManagerModule.cpp b/Engine/Source/Developer/TargetPlatform/Private/TargetPlatformManagerModule.cpp
--- a/Engine/Source/Developer/TargetPlatform/Private/TargetPlatformManagerModule.cpp
+++ b/Engine/Source/Developer/TargetPlatform/Private/TargetPlatformManagerModule.cpp
@@ -56,7 +56,7 @@ public:
 		bForceCacheUpdate = false;
 	}
 
-	virtual const TArray<ITargetPlatform*>& GetTargetPlatforms() OVERRIDE
+	virtual const TArray<ITargetPlatform*>& GetTargetPlatforms() override
 	{
 		if (Platforms.Num() == 0 || bForceCacheUpdate)
 		{
@@ -66,7 +66,7 @@ public:
 		return Platforms;
 	}
 
-	virtual ITargetDevicePtr FindTargetDevice(const FTargetDeviceId& DeviceId) OVERRIDE
+	virtual ITargetDevicePtr FindTargetDevice(const FTargetDeviceId& DeviceId) override
 	{
 		ITargetPlatform* Platform = FindTargetPlatform(DeviceId.GetPlatformName());
 
@@ -78,7 +78,7 @@ public:
 		return NULL;
 	}
 
-	virtual ITargetPlatform* FindTargetPlatform(FString Name) OVERRIDE
+	virtual ITargetPlatform* FindTargetPlatform(FString Name) override
 	{
 		const TArray<ITargetPlatform*>& TargetPlatforms = GetTargetPlatforms();	
 		
@@ -93,7 +93,7 @@ public:
 		return NULL;
 	}
 
-	virtual const TArray<ITargetPlatform*>& GetCookingTargetPlatforms() OVERRIDE
+	virtual const TArray<ITargetPlatform*>& GetCookingTargetPlatforms() override
 	{
 		static bool bInitialized = false;
 		static TArray<ITargetPlatform*> Results;
@@ -115,7 +115,7 @@ public:
 		return Results;
 	}
 
-	virtual const TArray<ITargetPlatform*>& GetActiveTargetPlatforms() OVERRIDE
+	virtual const TArray<ITargetPlatform*>& GetActiveTargetPlatforms() override
 	{
 		static bool bInitialized = false;
 		static TArray<ITargetPlatform*> Results;
@@ -191,14 +191,14 @@ public:
 		return Results;
 	}
 
-	virtual bool RestrictFormatsToRuntimeOnly() OVERRIDE
+	virtual bool RestrictFormatsToRuntimeOnly() override
 	{
 		GetActiveTargetPlatforms(); // make sure this is initialized
 
 		return bRestrictFormatsToRuntimeOnly;
 	}
 
-	virtual ITargetPlatform* GetRunningTargetPlatform() OVERRIDE
+	virtual ITargetPlatform* GetRunningTargetPlatform() override
 	{
 		static bool bInitialized = false;
 		static ITargetPlatform* Result = NULL;
@@ -228,7 +228,7 @@ public:
 		return Result;
 	}
 
-	virtual const TArray<const IAudioFormat*>& GetAudioFormats() OVERRIDE
+	virtual const TArray<const IAudioFormat*>& GetAudioFormats() override
 	{
 		static bool bInitialized = false;
 		static TArray<const IAudioFormat*> Results;
@@ -264,7 +264,7 @@ public:
 		return Results;
 	}
 
-	virtual const IAudioFormat* FindAudioFormat(FName Name) OVERRIDE
+	virtual const IAudioFormat* FindAudioFormat(FName Name) override
 	{
 		const TArray<const IAudioFormat*>& AudioFormats = GetAudioFormats();
 
@@ -286,7 +286,7 @@ public:
 		return NULL;
 	}
 
-	virtual const TArray<const ITextureFormat*>& GetTextureFormats() OVERRIDE
+	virtual const TArray<const ITextureFormat*>& GetTextureFormats() override
 	{
 		static bool bInitialized = false;
 		static TArray<const ITextureFormat*> Results;
@@ -322,7 +322,7 @@ public:
 		return Results;
 	}
 
-	virtual const ITextureFormat* FindTextureFormat(FName Name) OVERRIDE
+	virtual const ITextureFormat* FindTextureFormat(FName Name) override
 	{
 		const TArray<const ITextureFormat*>& TextureFormats = GetTextureFormats();
 
@@ -344,7 +344,7 @@ public:
 		return NULL;
 	}
 
-	virtual const TArray<const IShaderFormat*>& GetShaderFormats() OVERRIDE
+	virtual const TArray<const IShaderFormat*>& GetShaderFormats() override
 	{
 		static bool bInitialized = false;
 		static TArray<const IShaderFormat*> Results;
@@ -379,7 +379,7 @@ public:
 		return Results;
 	}
 
-	virtual const IShaderFormat* FindShaderFormat(FName Name) OVERRIDE
+	virtual const IShaderFormat* FindShaderFormat(FName Name) override
 	{
 		const TArray<const IShaderFormat*>& ShaderFormats = GetShaderFormats();	
 
@@ -401,7 +401,7 @@ public:
 		return NULL;
 	}
 
-	virtual uint16 ShaderFormatVersion(FName Name) OVERRIDE
+	virtual uint16 ShaderFormatVersion(FName Name) override
 	{
 		static TMap<FName, uint16> AlreadyFound;
 		uint16* Result = AlreadyFound.Find(Name);
@@ -421,7 +421,7 @@ public:
 		return *Result;
 	}
 
-	virtual const TArray<const IPhysXFormat*>& GetPhysXFormats() OVERRIDE
+	virtual const TArray<const IPhysXFormat*>& GetPhysXFormats() override
 	{
 		static bool bInitialized = false;
 		static TArray<const IPhysXFormat*> Results;
@@ -456,7 +456,7 @@ public:
 		return Results;
 	}
 
-	virtual const IPhysXFormat* FindPhysXFormat(FName Name) OVERRIDE
+	virtual const IPhysXFormat* FindPhysXFormat(FName Name) override
 	{
 		const TArray<const IPhysXFormat*>& PhysXFormats = GetPhysXFormats();
 
diff --git a/Engine/Source/Developer/TargetPlatform/Public/Common/TargetPlatformBase.h b/Engine/Source/Developer/TargetPlatform/Public/Common/TargetPlatformBase.h
--- a/Engine/Source/Developer/TargetPlatform/Public/Common/TargetPlatformBase.h
+++ b/Engine/Source/Developer/TargetPlatform/Public/Common/TargetPlatformBase.h
@@ -18,17 +18,17 @@ public:
 
 	// Begin ITargetPlatform interface
 
-	virtual bool AddDevice( const FString& DeviceName, bool bDefault ) OVERRIDE
+	virtual bool AddDevice( const FString& DeviceName, bool bDefault ) override
 	{
 		return false;
 	}
 
-	virtual FText DisplayName( ) const OVERRIDE
+	virtual FText DisplayName( ) const override
 	{
 		return PlatformInfo->DisplayName;
 	}
 
-	virtual const PlatformInfo::FPlatformInfo& GetPlatformInfo( ) const OVERRIDE
+	virtual const PlatformInfo::FPlatformInfo& GetPlatformInfo( ) const override
 	{
 		return *PlatformInfo;
 	}
@@ -156,12 +156,12 @@ public:
 #endif //TEXTURE_H_INCLUDED
 #endif //WITH_ENGINE
 
-	virtual bool PackageBuild( const FString& InPackgeDirectory ) OVERRIDE
+	virtual bool PackageBuild( const FString& InPackgeDirectory ) override
 	{
 		return true;
 	}
 
-	virtual bool IsSdkInstalled(bool bProjectHasCode, FString& OutDocumentationPath) const OVERRIDE
+	virtual bool IsSdkInstalled(bool bProjectHasCode, FString& OutDocumentationPath) const override
 	{
 		return true;
 	}
@@ -205,52 +205,52 @@ public:
 
 	// Begin ITargetPlatform interface
 
-	virtual bool HasEditorOnlyData( ) const OVERRIDE
+	virtual bool HasEditorOnlyData( ) const override
 	{
 		return TPlatformProperties::HasEditorOnlyData();
 	}
 
-	virtual bool IsLittleEndian( ) const OVERRIDE
+	virtual bool IsLittleEndian( ) const override
 	{
 		return TPlatformProperties::IsLittleEndian();
 	}
 
-	virtual bool IsServerOnly( ) const OVERRIDE
+	virtual bool IsServerOnly( ) const override
 	{
 		return TPlatformProperties::IsServerOnly();
 	}
 
-	virtual bool IsClientOnly( ) const OVERRIDE
+	virtual bool IsClientOnly( ) const override
 	{
 		return TPlatformProperties::IsClientOnly();
 	}
 
-	virtual FString PlatformName( ) const OVERRIDE
+	virtual FString PlatformName( ) const override
 	{
 		return FString(TPlatformProperties::PlatformName());
 	}
 
-	virtual FString IniPlatformName( ) const OVERRIDE
+	virtual FString IniPlatformName( ) const override
 	{
 		return FString(TPlatformProperties::IniPlatformName());
 	}
 
-	virtual bool RequiresCookedData( ) const OVERRIDE
+	virtual bool RequiresCookedData( ) const override
 	{
 		return TPlatformProperties::RequiresCookedData();
 	}
 
-	virtual bool RequiresUserCredentials() const OVERRIDE
+	virtual bool RequiresUserCredentials() const override
 	{
 		return TPlatformProperties::RequiresUserCredentials();
 	}
 
-	virtual bool SupportsBuildTarget( EBuildTargets::Type BuildTarget ) const OVERRIDE
+	virtual bool SupportsBuildTarget( EBuildTargets::Type BuildTarget ) const override
 	{
 		return TPlatformProperties::SupportsBuildTarget(BuildTarget);
 	}
 
-	virtual bool SupportsFeature( ETargetPlatformFeatures::Type Feature ) const OVERRIDE
+	virtual bool SupportsFeature( ETargetPlatformFeatures::Type Feature ) const override
 	{
 		switch (Feature)
 		{
@@ -285,13 +285,13 @@ public:
 		return false;
 	}
 	
-	virtual uint32 MaxGpuSkinBones( ) const OVERRIDE
+	virtual uint32 MaxGpuSkinBones( ) const override
 	{
 		return TPlatformProperties::MaxGpuSkinBones();
 	}
 
 #if WITH_ENGINE
-	virtual FName GetPhysicsFormat( class UBodySetup* Body ) const OVERRIDE
+	virtual FName GetPhysicsFormat( class UBodySetup* Body ) const override
 	{
 		return FName(TPlatformProperties::GetPhysicsFormat());
 	}
diff --git a/Engine/Source/Developer/TaskGraph/Private/STaskGraph.cpp b/Engine/Source/Developer/TaskGraph/Private/STaskGraph.cpp
--- a/Engine/Source/Developer/TaskGraph/Private/STaskGraph.cpp
+++ b/Engine/Source/Developer/TaskGraph/Private/STaskGraph.cpp
@@ -94,7 +94,7 @@ public:
 		return true;
 	}
 
-	virtual TStatId GetStatId() const OVERRIDE
+	virtual TStatId GetStatId() const override
 	{
 		RETURN_QUICK_DECLARE_CYCLE_STAT(FDelayedVisualizerSpawner, STATGROUP_Tickables);
 	}
@@ -162,15 +162,15 @@ void DisplayProfileVisualizer( TSharedPtr< FVisualizerEvent > InProfileData, con
 class FProfileVisualizerModule : public IProfileVisualizerModule
 {
 public:
-	virtual void StartupModule() OVERRIDE
+	virtual void StartupModule() override
 	{
 		::InitProfileVisualizer();
 	}
-	virtual void ShutdownModule() OVERRIDE
+	virtual void ShutdownModule() override
 	{
 		::ShutdownProfileVisualizer();
 	}
-	virtual void DisplayProfileVisualizer( TSharedPtr< FVisualizerEvent > InProfileData, const TCHAR* InProfilerType ) OVERRIDE
+	virtual void DisplayProfileVisualizer( TSharedPtr< FVisualizerEvent > InProfileData, const TCHAR* InProfilerType ) override
 	{
 #if	WITH_EDITOR
 		::DisplayProfileVisualizer( InProfileData, InProfilerType );
diff --git a/Engine/Source/Developer/TaskGraph/Public/SEventsTree.h b/Engine/Source/Developer/TaskGraph/Public/SEventsTree.h
--- a/Engine/Source/Developer/TaskGraph/Public/SEventsTree.h
+++ b/Engine/Source/Developer/TaskGraph/Public/SEventsTree.h
@@ -26,7 +26,7 @@ public:
 	 *
 	 * @param InArgs   A declaration from which to construct the widget
 	 */
-	virtual TSharedRef< SWidget > GenerateWidgetForColumn( const FName& ColumnName ) OVERRIDE;
+	virtual TSharedRef< SWidget > GenerateWidgetForColumn( const FName& ColumnName ) override;
 
 
 	/**
diff --git a/Engine/Source/Developer/TaskGraph/Public/SGraphBar.h b/Engine/Source/Developer/TaskGraph/Public/SGraphBar.h
--- a/Engine/Source/Developer/TaskGraph/Public/SGraphBar.h
+++ b/Engine/Source/Developer/TaskGraph/Public/SGraphBar.h
@@ -33,11 +33,11 @@ public:
 	void Construct( const FArguments& InArgs );
 
 	// SWidget interface
-	virtual int32 OnPaint( const FGeometry& AllottedGeometry, const FSlateRect& MyClippingRect, FSlateWindowElementList& OutDrawElements, int32 LayerId, const FWidgetStyle& InWidgetStyle, bool bParentEnabled ) const OVERRIDE;
-	virtual void Tick( const FGeometry& AllottedGeometry, const double InCurrentTime, const float InDeltaTime ) OVERRIDE;
-	virtual FReply OnMouseButtonDown( const FGeometry& MyGeometry, const FPointerEvent& MouseEvent ) OVERRIDE;
-	virtual FReply OnMouseMove( const FGeometry& MyGeometry, const FPointerEvent& MouseEvent ) OVERRIDE;
-	FVector2D ComputeDesiredSize() const OVERRIDE;
+	virtual int32 OnPaint( const FGeometry& AllottedGeometry, const FSlateRect& MyClippingRect, FSlateWindowElementList& OutDrawElements, int32 LayerId, const FWidgetStyle& InWidgetStyle, bool bParentEnabled ) const override;
+	virtual void Tick( const FGeometry& AllottedGeometry, const double InCurrentTime, const float InDeltaTime ) override;
+	virtual FReply OnMouseButtonDown( const FGeometry& MyGeometry, const FPointerEvent& MouseEvent ) override;
+	virtual FReply OnMouseMove( const FGeometry& MyGeometry, const FPointerEvent& MouseEvent ) override;
+	FVector2D ComputeDesiredSize() const override;
 	// End of SWidget interface
 
 	/**
diff --git a/Engine/Source/Developer/TaskGraph/Public/STimeline.h b/Engine/Source/Developer/TaskGraph/Public/STimeline.h
--- a/Engine/Source/Developer/TaskGraph/Public/STimeline.h
+++ b/Engine/Source/Developer/TaskGraph/Public/STimeline.h
@@ -31,11 +31,11 @@ public:
 	void Construct( const FArguments& InArgs );
 
 	// SWidget interface
-	virtual int32 OnPaint( const FGeometry& AllottedGeometry, const FSlateRect& MyClippingRect, FSlateWindowElementList& OutDrawElements, int32 LayerId, const FWidgetStyle& InWidgetStyle, bool bParentEnabled ) const OVERRIDE;
-	virtual void Tick( const FGeometry& AllottedGeometry, const double InCurrentTime, const float InDeltaTime ) OVERRIDE;
-	virtual FReply OnMouseButtonDown( const FGeometry& MyGeometry, const FPointerEvent& MouseEvent ) OVERRIDE;
-	virtual FReply OnMouseMove( const FGeometry& MyGeometry, const FPointerEvent& MouseEvent ) OVERRIDE;
-	FVector2D ComputeDesiredSize() const OVERRIDE;
+	virtual int32 OnPaint( const FGeometry& AllottedGeometry, const FSlateRect& MyClippingRect, FSlateWindowElementList& OutDrawElements, int32 LayerId, const FWidgetStyle& InWidgetStyle, bool bParentEnabled ) const override;
+	virtual void Tick( const FGeometry& AllottedGeometry, const double InCurrentTime, const float InDeltaTime ) override;
+	virtual FReply OnMouseButtonDown( const FGeometry& MyGeometry, const FPointerEvent& MouseEvent ) override;
+	virtual FReply OnMouseMove( const FGeometry& MyGeometry, const FPointerEvent& MouseEvent ) override;
+	FVector2D ComputeDesiredSize() const override;
 	// End of SWidget interface
 
 	/**
diff --git a/Engine/Source/Developer/TextureFormatAndroid/Private/TextureFormatAndroid.cpp b/Engine/Source/Developer/TextureFormatAndroid/Private/TextureFormatAndroid.cpp
--- a/Engine/Source/Developer/TextureFormatAndroid/Private/TextureFormatAndroid.cpp
+++ b/Engine/Source/Developer/TextureFormatAndroid/Private/TextureFormatAndroid.cpp
@@ -123,17 +123,17 @@ static bool CompressImageUsingQonvert(
  */
 class FTextureFormatAndroid : public ITextureFormat
 {
-	virtual bool AllowParallelBuild() const OVERRIDE
+	virtual bool AllowParallelBuild() const override
 	{
 		return true;
 	}
 
-	virtual uint16 GetVersion(FName Format) const OVERRIDE
+	virtual uint16 GetVersion(FName Format) const override
 	{
 		return 0;
 	}
 
-	virtual void GetSupportedFormats(TArray<FName>& OutFormats) const OVERRIDE
+	virtual void GetSupportedFormats(TArray<FName>& OutFormats) const override
 	{
 		for (int32 i = 0; i < ARRAY_COUNT(GSupportedTextureFormatNames); ++i)
 		{
@@ -141,7 +141,7 @@ class FTextureFormatAndroid : public ITextureFormat
 		}
 	}
 
-	virtual FTextureFormatCompressorCaps GetFormatCapabilities() const OVERRIDE
+	virtual FTextureFormatCompressorCaps GetFormatCapabilities() const override
 	{
 		return FTextureFormatCompressorCaps(); // Default capabilities.
 	}
@@ -151,7 +151,7 @@ class FTextureFormatAndroid : public ITextureFormat
 		const struct FTextureBuildSettings& BuildSettings,
 		bool bImageHasAlphaChannel,
 		FCompressedImage2D& OutCompressedImage
-		) const OVERRIDE
+		) const override
 	{
 		FImage Image;
 		InImage.CopyTo(Image, ERawImageFormat::BGRA8, BuildSettings.bSRGB);
diff --git a/Engine/Source/Developer/TextureFormatDXT/Private/TextureFormatDXT.cpp b/Engine/Source/Developer/TextureFormatDXT/Private/TextureFormatDXT.cpp
--- a/Engine/Source/Developer/TextureFormatDXT/Private/TextureFormatDXT.cpp
+++ b/Engine/Source/Developer/TextureFormatDXT/Private/TextureFormatDXT.cpp
@@ -387,17 +387,17 @@ static bool CompressImageUsingNVTT(
  */
 class FTextureFormatDXT : public ITextureFormat
 {
-	virtual bool AllowParallelBuild() const OVERRIDE
+	virtual bool AllowParallelBuild() const override
 	{
 		return true;
 	}
 
-	virtual uint16 GetVersion(FName Format) const OVERRIDE
+	virtual uint16 GetVersion(FName Format) const override
 	{
 		return 0;
 	}
 
-	virtual void GetSupportedFormats(TArray<FName>& OutFormats) const OVERRIDE
+	virtual void GetSupportedFormats(TArray<FName>& OutFormats) const override
 	{
 		for (int32 i = 0; i < ARRAY_COUNT(GSupportedTextureFormatNames); ++i)
 		{
@@ -405,7 +405,7 @@ class FTextureFormatDXT : public ITextureFormat
 		}
 	}
 	
-	virtual FTextureFormatCompressorCaps GetFormatCapabilities() const OVERRIDE
+	virtual FTextureFormatCompressorCaps GetFormatCapabilities() const override
 	{
 		return FTextureFormatCompressorCaps(); // Default capabilities.
 	}
@@ -415,7 +415,7 @@ class FTextureFormatDXT : public ITextureFormat
 		const struct FTextureBuildSettings& BuildSettings,
 		bool bImageHasAlphaChannel,
 		FCompressedImage2D& OutCompressedImage
-		) const OVERRIDE
+		) const override
 	{
 		FImage Image;
 		InImage.CopyTo(Image, ERawImageFormat::BGRA8, BuildSettings.bSRGB);
