Slate Hittesting 2.0 : Hittesting no longer bound to logical widget hierarchy.

Merged from usr/Nick.Atamas/UE4_SlateDev => /UE4

diff --git a/Engine/Plugins/Experimental/Paper2D/Source/Paper2DEditor/Private/FlipbookEditor/SFlipbookTimeline.cpp b/Engine/Plugins/Experimental/Paper2D/Source/Paper2DEditor/Private/FlipbookEditor/SFlipbookTimeline.cpp
--- a/Engine/Plugins/Experimental/Paper2D/Source/Paper2DEditor/Private/FlipbookEditor/SFlipbookTimeline.cpp
+++ b/Engine/Plugins/Experimental/Paper2D/Source/Paper2DEditor/Private/FlipbookEditor/SFlipbookTimeline.cpp
@@ -134,9 +134,9 @@ void SFlipbookTimeline::OnAssetsDropped(const class FAssetDragDropOp& DragDropOp
 	}
 }
 
-int32 SFlipbookTimeline::OnPaint(const FGeometry& AllottedGeometry, const FSlateRect& MyClippingRect, FSlateWindowElementList& OutDrawElements, int32 LayerId, const FWidgetStyle& InWidgetStyle, bool bParentEnabled) const
+int32 SFlipbookTimeline::OnPaint(const FPaintArgs& Args, const FGeometry& AllottedGeometry, const FSlateRect& MyClippingRect, FSlateWindowElementList& OutDrawElements, int32 LayerId, const FWidgetStyle& InWidgetStyle, bool bParentEnabled) const
 {
-	LayerId = SCompoundWidget::OnPaint(AllottedGeometry, MyClippingRect, OutDrawElements, LayerId, InWidgetStyle, bParentEnabled);
+	LayerId = SCompoundWidget::OnPaint(Args, AllottedGeometry, MyClippingRect, OutDrawElements, LayerId, InWidgetStyle, bParentEnabled);
 
 	const float CurrentTimeSecs = PlayTime.Get();
 	UPaperFlipbook* Flipbook = FlipbookBeingEdited.Get();
diff --git a/Engine/Plugins/Experimental/Paper2D/Source/Paper2DEditor/Private/FlipbookEditor/SFlipbookTimeline.h b/Engine/Plugins/Experimental/Paper2D/Source/Paper2DEditor/Private/FlipbookEditor/SFlipbookTimeline.h
--- a/Engine/Plugins/Experimental/Paper2D/Source/Paper2DEditor/Private/FlipbookEditor/SFlipbookTimeline.h
+++ b/Engine/Plugins/Experimental/Paper2D/Source/Paper2DEditor/Private/FlipbookEditor/SFlipbookTimeline.h
@@ -24,7 +24,7 @@ public:
 
 	// SWidget interface
 	virtual FReply OnDrop(const FGeometry& MyGeometry, const FDragDropEvent& DragDropEvent) override;
-	virtual int32 OnPaint(const FGeometry& AllottedGeometry, const FSlateRect& MyClippingRect, FSlateWindowElementList& OutDrawElements, int32 LayerId, const FWidgetStyle& InWidgetStyle, bool bParentEnabled) const override;
+	virtual int32 OnPaint(const FPaintArgs& Args, const FGeometry& AllottedGeometry, const FSlateRect& MyClippingRect, FSlateWindowElementList& OutDrawElements, int32 LayerId, const FWidgetStyle& InWidgetStyle, bool bParentEnabled) const override;
 	virtual FReply OnMouseButtonUp(const FGeometry& MyGeometry, const FPointerEvent& MouseEvent) override;
 	// End of SWidget interface
 
diff --git a/Engine/Plugins/Experimental/Paper2D/Source/Paper2DEditor/Private/SPaperEditorViewport.cpp b/Engine/Plugins/Experimental/Paper2D/Source/Paper2DEditor/Private/SPaperEditorViewport.cpp
--- a/Engine/Plugins/Experimental/Paper2D/Source/Paper2DEditor/Private/SPaperEditorViewport.cpp
+++ b/Engine/Plugins/Experimental/Paper2D/Source/Paper2DEditor/Private/SPaperEditorViewport.cpp
@@ -599,9 +599,9 @@ void SPaperEditorViewport::PaintSoftwareCursor(const FGeometry& AllottedGeometry
 	}
 }
 
-int32 SPaperEditorViewport::OnPaint(const FGeometry& AllottedGeometry, const FSlateRect& MyClippingRect, FSlateWindowElementList& OutDrawElements, int32 LayerId, const FWidgetStyle& InWidgetStyle, bool bParentEnabled) const
+int32 SPaperEditorViewport::OnPaint(const FPaintArgs& Args, const FGeometry& AllottedGeometry, const FSlateRect& MyClippingRect, FSlateWindowElementList& OutDrawElements, int32 LayerId, const FWidgetStyle& InWidgetStyle, bool bParentEnabled) const
 {
-	int32 MaxLayerId = SCompoundWidget::OnPaint(AllottedGeometry, MyClippingRect, OutDrawElements, LayerId, InWidgetStyle, bParentEnabled);
+	int32 MaxLayerId = SCompoundWidget::OnPaint(Args, AllottedGeometry, MyClippingRect, OutDrawElements, LayerId, InWidgetStyle, bParentEnabled);
 
 	++MaxLayerId;
 	PaintSoftwareCursor(AllottedGeometry, MyClippingRect, OutDrawElements, MaxLayerId);
diff --git a/Engine/Plugins/Experimental/Paper2D/Source/Paper2DEditor/Private/SPaperEditorViewport.h b/Engine/Plugins/Experimental/Paper2D/Source/Paper2DEditor/Private/SPaperEditorViewport.h
--- a/Engine/Plugins/Experimental/Paper2D/Source/Paper2DEditor/Private/SPaperEditorViewport.h
+++ b/Engine/Plugins/Experimental/Paper2D/Source/Paper2DEditor/Private/SPaperEditorViewport.h
@@ -86,7 +86,7 @@ public:
 	virtual FReply OnMouseMove(const FGeometry& MyGeometry, const FPointerEvent& MouseEvent) override;
 	virtual FReply OnMouseWheel(const FGeometry& MyGeometry, const FPointerEvent& MouseEvent) override;
 	virtual FCursorReply OnCursorQuery( const FGeometry& MyGeometry, const FPointerEvent& CursorEvent ) const override;
-	virtual int32 OnPaint(const FGeometry& AllottedGeometry, const FSlateRect& MyClippingRect, FSlateWindowElementList& OutDrawElements, int32 LayerId, const FWidgetStyle& InWidgetStyle, bool bParentEnabled) const override;
+	virtual int32 OnPaint(const FPaintArgs& Args, const FGeometry& AllottedGeometry, const FSlateRect& MyClippingRect, FSlateWindowElementList& OutDrawElements, int32 LayerId, const FWidgetStyle& InWidgetStyle, bool bParentEnabled) const override;
 	// End of SWidget interface
 
 	void Construct(const FArguments& InArgs, TSharedRef<class FPaperEditorViewportClient> InViewportClient);
diff --git a/Engine/Source/Developer/LogVisualizer/Private/SLogBar.cpp b/Engine/Source/Developer/LogVisualizer/Private/SLogBar.cpp
--- a/Engine/Source/Developer/LogVisualizer/Private/SLogBar.cpp
+++ b/Engine/Source/Developer/LogVisualizer/Private/SLogBar.cpp
@@ -32,7 +32,7 @@ void SLogBar::Construct( const FArguments& InArgs )
 	HistogramPreviewWindow = 50;
 }
 
-int32 SLogBar::OnPaint( const FGeometry& AllottedGeometry, const FSlateRect& MyClippingRect, FSlateWindowElementList& OutDrawElements, int32 LayerId, const FWidgetStyle& InWidgetStyle, bool bParentEnabled ) const
+int32 SLogBar::OnPaint( const FPaintArgs& Args, const FGeometry& AllottedGeometry, const FSlateRect& MyClippingRect, FSlateWindowElementList& OutDrawElements, int32 LayerId, const FWidgetStyle& InWidgetStyle, bool bParentEnabled ) const
 {
 	// Used to track the layer ID we will return.
 	int32 RetLayerId = LayerId;
diff --git a/Engine/Source/Developer/LogVisualizer/Private/SLogBar.h b/Engine/Source/Developer/LogVisualizer/Private/SLogBar.h
--- a/Engine/Source/Developer/LogVisualizer/Private/SLogBar.h
+++ b/Engine/Source/Developer/LogVisualizer/Private/SLogBar.h
@@ -47,60 +47,18 @@ public:
 	 */
 	void Construct(const FArguments& InArgs);
 
-	/**
-	 * The widget should respond by populating the OutDrawElements array with FDrawElements 
-	 * that represent it and any of its children.
-	 *
-	 * @param AllottedGeometry  The FGeometry that describes an area in which the widget should appear.
-	 * @param MyClippingRect    The clipping rectangle allocated for this widget and its children.
-	 * @param OutDrawElements   A list of FDrawElements to populate with the output.
-	 * @param LayerId           The Layer onto which this widget should be rendered.
-	 * @param InColorAndOpacity Color and Opacity to be applied to all the descendants of the widget being painted
- 	 * @param bParentEnabled	True if the parent of this widget is enabled.
-	 *
-	 * @return The maximum layer ID attained by this widget or any of its children.
-	 */
-	virtual int32 OnPaint(const FGeometry& AllottedGeometry, const FSlateRect& MyClippingRect, FSlateWindowElementList& OutDrawElements, int32 LayerId, const FWidgetStyle& InWidgetStyle, bool bParentEnabled) const override;
+	virtual int32 OnPaint(const FPaintArgs& Args, const FGeometry& AllottedGeometry, const FSlateRect& MyClippingRect, FSlateWindowElementList& OutDrawElements, int32 LayerId, const FWidgetStyle& InWidgetStyle, bool bParentEnabled) const override;
 	
-	/**
-	 * Ticks this widget.  Override in derived classes, but always call the parent implementation.
-	 *
-	 * @param  AllottedGeometry The space allotted for this widget
-	 * @param  InCurrentTime  Current absolute real time
-	 * @param  InDeltaTime  Real time passed since last tick
-	 */
 	virtual void Tick(const FGeometry& AllottedGeometry, const double InCurrentTime, const float InDeltaTime);
 
-	/**
-	 * The system calls this method to notify the widget that a mouse button was pressed within it. This event is bubbled.
-	 *
-	 * @param MyGeometry The Geometry of the widget receiving the event
-	 * @param MouseEvent Information about the input event
-	 *
-	 * @return Whether the event was handled along with possible requests for the system to take action.
-	 */
 	virtual FReply OnMouseButtonUp(const FGeometry& MyGeometry, const FPointerEvent& MouseEvent) override;
 
 	virtual FReply OnMouseButtonDown(const FGeometry& MyGeometry, const FPointerEvent& MouseEvent) override;
 
-	/**
-	 * The system calls this method to notify the widget that a mouse moved within it. This event is bubbled.
-	 *
-	 * @param MyGeometry The Geometry of the widget receiving the event
-	 * @param MouseEvent Information about the input event
-	 *
-	 * @return Whether the event was handled along with possible requests for the system to take action.
-	 */
 	virtual FReply OnMouseMove(const FGeometry& MyGeometry, const FPointerEvent& MouseEvent) override;
 	
 	void OnCurrentTimeChanged(float NewTime);
 
-	/**
-	 * A Panel's desired size in the space required to arrange of its children on the screen while respecting all of
-	 * the children's desired sizes and any layout-related options specified by the user. See StackPanel for an example.
-	 *
-	 * @return The desired size.
-	 */
 	FVector2D ComputeDesiredSize() const;
 
 	/**
diff --git a/Engine/Source/Developer/Profiler/Private/Widgets/SDataGraph.cpp b/Engine/Source/Developer/Profiler/Private/Widgets/SDataGraph.cpp
--- a/Engine/Source/Developer/Profiler/Private/Widgets/SDataGraph.cpp
+++ b/Engine/Source/Developer/Profiler/Private/Widgets/SDataGraph.cpp
@@ -283,15 +283,7 @@ void SDataGraph::UpdateState()
 	}
 }
 
-int32 SDataGraph::OnPaint
-( 
-	const FGeometry& AllottedGeometry, 
-	const FSlateRect& MyClippingRect, 
-	FSlateWindowElementList& OutDrawElements, 
-	int32 LayerId, 
-	const FWidgetStyle& InWidgetStyle, 
-	bool bParentEnabled 
-) const
+int32 SDataGraph::OnPaint( const FPaintArgs& Args, const FGeometry& AllottedGeometry, const FSlateRect& MyClippingRect, FSlateWindowElementList& OutDrawElements, int32 LayerId, const FWidgetStyle& InWidgetStyle, bool bParentEnabled ) const
 {
 	static double TotalTime = 0.0f;
 	static uint32 NumCalls = 0;
@@ -978,7 +970,7 @@ int32 SDataGraph::OnPaint
 	}
 #endif // DEBUG_PROFILER_PERFORMANCE
 
-	return SCompoundWidget::OnPaint(AllottedGeometry, MyClippingRect, OutDrawElements, LayerId, InWidgetStyle, bParentEnabled && IsEnabled() );
+	return SCompoundWidget::OnPaint( Args, AllottedGeometry, MyClippingRect, OutDrawElements, LayerId, InWidgetStyle, bParentEnabled && IsEnabled() );
 }
 
 
diff --git a/Engine/Source/Developer/Profiler/Private/Widgets/SDataGraph.h b/Engine/Source/Developer/Profiler/Private/Widgets/SDataGraph.h
--- a/Engine/Source/Developer/Profiler/Private/Widgets/SDataGraph.h
+++ b/Engine/Source/Developer/Profiler/Private/Widgets/SDataGraph.h
@@ -138,20 +138,7 @@ public:
 	 */
 	virtual void Tick( const FGeometry& AllottedGeometry, const double InCurrentTime, const float InDeltaTime ) override;
 
-	/**
-	 * The widget should respond by populating the OutDrawElements array with FDrawElements 
-	 * that represent it and any of its children.
-	 *
-	 * @param AllottedGeometry  The FGeometry that describes an area in which the widget should appear.
-	 * @param MyClippingRect    The clipping rectangle allocated for this widget and its children.
-	 * @param OutDrawElements   A list of FDrawElements to populate with the output.
-	 * @param LayerId           The Layer onto which this widget should be rendered.
-	 * @param InColorAndOpacity ColorAverage and Opacity to be applied to all the descendants of the widget being painted
-	 * @param bParentEnabled	True if the parent of this widget is enabled.
-	 *
-	 * @return The maximum layer ID attained by this widget or any of its children.
-	 */
-	virtual int32 OnPaint( const FGeometry& AllottedGeometry, const FSlateRect& MyClippingRect, FSlateWindowElementList& OutDrawElements, int32 LayerId, const FWidgetStyle& InWidgetStyle, bool bParentEnabled ) const override;
+	virtual int32 OnPaint( const FPaintArgs& Args, const FGeometry& AllottedGeometry, const FSlateRect& MyClippingRect, FSlateWindowElementList& OutDrawElements, int32 LayerId, const FWidgetStyle& InWidgetStyle, bool bParentEnabled ) const override;
 
 	/**
 	 * The system calls this method to notify the widget that a mouse button was pressed within it. This event is bubbled.
diff --git a/Engine/Source/Developer/Profiler/Private/Widgets/SHistogram.cpp b/Engine/Source/Developer/Profiler/Private/Widgets/SHistogram.cpp
--- a/Engine/Source/Developer/Profiler/Private/Widgets/SHistogram.cpp
+++ b/Engine/Source/Developer/Profiler/Private/Widgets/SHistogram.cpp
@@ -34,15 +34,7 @@ void SHistogram::Tick( const FGeometry& AllottedGeometry, const double InCurrent
 	SCompoundWidget::Tick(AllottedGeometry,InCurrentTime,InDeltaTime);
 }
 
-int32 SHistogram::OnPaint
-( 
-	const FGeometry& AllottedGeometry, 
-	const FSlateRect& MyClippingRect, 
-	FSlateWindowElementList& OutDrawElements, 
-	int32 LayerId, 
-	const FWidgetStyle& InWidgetStyle, 
-	bool bParentEnabled 
-) const
+int32 SHistogram::OnPaint( const FPaintArgs& Args, const FGeometry& AllottedGeometry, const FSlateRect& MyClippingRect, FSlateWindowElementList& OutDrawElements, int32 LayerId, const FWidgetStyle& InWidgetStyle, bool bParentEnabled ) const
 {
 	static double TotalTime = 0.0f;
 	static uint32 NumCalls = 0;
@@ -188,7 +180,7 @@ int32 SHistogram::OnPaint
 			FLinearColor::Green
 			);
 	}
-	return SCompoundWidget::OnPaint(AllottedGeometry, MyClippingRect, OutDrawElements, LayerId, InWidgetStyle, bParentEnabled && IsEnabled() );
+	return SCompoundWidget::OnPaint( Args, AllottedGeometry, MyClippingRect, OutDrawElements, LayerId, InWidgetStyle, bParentEnabled && IsEnabled() );
 }
 
 void SHistogram::RestoreDefaultState()
diff --git a/Engine/Source/Developer/Profiler/Private/Widgets/SHistogram.h b/Engine/Source/Developer/Profiler/Private/Widgets/SHistogram.h
--- a/Engine/Source/Developer/Profiler/Private/Widgets/SHistogram.h
+++ b/Engine/Source/Developer/Profiler/Private/Widgets/SHistogram.h
@@ -115,99 +115,24 @@ public:
 	 */
 	void Construct( const FArguments& InArgs );
 
-	/**
-	 * Ticks this widget.  Override in derived classes, but always call the parent implementation.
-	 *
-	 * @param  AllottedGeometry The space allotted for this widget
-	 * @param  InCurrentTime  Current absolute real time
-	 * @param  InDeltaTime  Real time passed since last tick
-	 */
 	virtual void Tick( const FGeometry& AllottedGeometry, const double InCurrentTime, const float InDeltaTime ) override;
 
-	/**
-	 * The widget should respond by populating the OutDrawElements array with FDrawElements 
-	 * that represent it and any of its children.
-	 *
-	 * @param AllottedGeometry  The FGeometry that describes an area in which the widget should appear.
-	 * @param MyClippingRect    The clipping rectangle allocated for this widget and its children.
-	 * @param OutDrawElements   A list of FDrawElements to populate with the output.
-	 * @param LayerId           The Layer onto which this widget should be rendered.
-	 * @param InColorAndOpacity Color and Opacity to be applied to all the descendants of the widget being painted
-	 * @param bParentEnabled	True if the parent of this widget is enabled.
-	 *
-	 * @return The maximum layer ID attained by this widget or any of its children.
-	 */
-	virtual int32 OnPaint( const FGeometry& AllottedGeometry, const FSlateRect& MyClippingRect, FSlateWindowElementList& OutDrawElements, int32 LayerId, const FWidgetStyle& InWidgetStyle, bool bParentEnabled ) const override;
+	virtual int32 OnPaint( const FPaintArgs& Args, const FGeometry& AllottedGeometry, const FSlateRect& MyClippingRect, FSlateWindowElementList& OutDrawElements, int32 LayerId, const FWidgetStyle& InWidgetStyle, bool bParentEnabled ) const override;
 
-	/**
-	 * The system calls this method to notify the widget that a mouse button was pressed within it. This event is bubbled.
-	 *
-	 * @param MyGeometry The Geometry of the widget receiving the event
-	 * @param MouseEvent Information about the input event
-	 *
-	 * @return Whether the event was handled along with possible requests for the system to take action.
-	 */
 	virtual FReply OnMouseButtonDown( const FGeometry& MyGeometry, const FPointerEvent& MouseEvent ) override;
 	
-	/**
-	 * The system calls this method to notify the widget that a mouse button was release within it. This event is bubbled.
-	 *
-	 * @param MyGeometry The Geometry of the widget receiving the event
-	 * @param MouseEvent Information about the input event
-	 *
-	 * @return Whether the event was handled along with possible requests for the system to take action.
-	 */
 	virtual FReply OnMouseButtonUp( const FGeometry& MyGeometry, const FPointerEvent& MouseEvent ) override;
 	
-	/**
-	 * The system calls this method to notify the widget that a mouse moved within it. This event is bubbled.
-	 *
-	 * @param MyGeometry The Geometry of the widget receiving the event
-	 * @param MouseEvent Information about the input event
-	 *
-	 * @return Whether the event was handled along with possible requests for the system to take action.
-	 */
 	virtual FReply OnMouseMove( const FGeometry& MyGeometry, const FPointerEvent& MouseEvent ) override;
 	
-	/**
-	 * The system will use this event to notify a widget that the cursor has entered it. This event is NOT bubbled.
-	 *
-	 * @param MyGeometry The Geometry of the widget receiving the event
-	 * @param MouseEvent Information about the input event
-	 */
 	virtual void OnMouseEnter( const FGeometry& MyGeometry, const FPointerEvent& MouseEvent ) override;
 	
-	/**
-	 * The system will use this event to notify a widget that the cursor has left it. This event is NOT bubbled.
-	 *
-	 * @param MouseEvent Information about the input event
-	 */
 	virtual void OnMouseLeave( const FPointerEvent& MouseEvent ) override;
 	
-	/**
-	 * Called when the mouse wheel is spun. This event is bubbled.
-	 *
-	 * @param  MouseEvent  Mouse event
-	 *
-	 * @return  Returns whether the event was handled, along with other possible actions
-	 */
 	virtual FReply OnMouseWheel( const FGeometry& MyGeometry, const FPointerEvent& MouseEvent ) override;
 
-	/**
-	 * Called when a mouse button is double clicked.  Override this in derived classes.
-	 *
-	 * @param  InMyGeometry  Widget geometry
-	 * @param  InMouseEvent  Mouse button event
-	 *
-	 * @return  Returns whether the event was handled, along with other possible actions
-	 */
 	virtual FReply OnMouseButtonDoubleClick( const FGeometry& MyGeometry, const FPointerEvent& MouseEvent ) override;
 
-	/**
-	 * Called when the system wants to know which cursor to display for this Widget.  This event is bubbled.
-	 *
-	 * @return  The cursor requested (can be None.)
-	 */
 	virtual FCursorReply OnCursorQuery( const FGeometry& MyGeometry, const FPointerEvent& CursorEvent ) const override;
 
 	/** Restores default state for the data graph by removing all inner graphs. */
diff --git a/Engine/Source/Developer/Profiler/Private/Widgets/SProfilerMiniView.cpp b/Engine/Source/Developer/Profiler/Private/Widgets/SProfilerMiniView.cpp
--- a/Engine/Source/Developer/Profiler/Private/Widgets/SProfilerMiniView.cpp
+++ b/Engine/Source/Developer/Profiler/Private/Widgets/SProfilerMiniView.cpp
@@ -82,7 +82,7 @@ void SProfilerMiniView::Tick( const FGeometry& AllottedGeometry, const double In
 	}
 }
 
-int32 SProfilerMiniView::OnPaint( const FGeometry& AllottedGeometry, const FSlateRect& MyClippingRect, FSlateWindowElementList& OutDrawElements, int32 LayerId, const FWidgetStyle& InWidgetStyle, bool bParentEnabled ) const
+int32 SProfilerMiniView::OnPaint( const FPaintArgs& Args, const FGeometry& AllottedGeometry, const FSlateRect& MyClippingRect, FSlateWindowElementList& OutDrawElements, int32 LayerId, const FWidgetStyle& InWidgetStyle, bool bParentEnabled ) const
 {
 //	SCOPE_LOG_TIME_FUNC();
 
@@ -320,7 +320,7 @@ int32 SProfilerMiniView::OnPaint( const FGeometry& AllottedGeometry, const FSlat
 
 #endif // DEBUG_PROFILER_PERFORMANCE
 
-	return SCompoundWidget::OnPaint( AllottedGeometry, MyClippingRect, OutDrawElements, LayerId, InWidgetStyle, bParentEnabled && IsEnabled() );
+	return SCompoundWidget::OnPaint( Args, AllottedGeometry, MyClippingRect, OutDrawElements, LayerId, InWidgetStyle, bParentEnabled && IsEnabled() );
 }
 
 
diff --git a/Engine/Source/Developer/Profiler/Private/Widgets/SProfilerMiniView.h b/Engine/Source/Developer/Profiler/Private/Widgets/SProfilerMiniView.h
--- a/Engine/Source/Developer/Profiler/Private/Widgets/SProfilerMiniView.h
+++ b/Engine/Source/Developer/Profiler/Private/Widgets/SProfilerMiniView.h
@@ -204,99 +204,24 @@ public:
 	/** Resets internal widget's data to the default one. */
 	void Reset();
 
-	/**
-	* Ticks this widget.  Override in derived classes, but always call the parent implementation.
-	*
-	* @param  AllottedGeometry The space allotted for this widget
-	* @param  InCurrentTime  Current absolute real time
-	* @param  InDeltaTime  Real time passed since last tick
-	*/
 	virtual void Tick( const FGeometry& AllottedGeometry, const double InCurrentTime, const float InDeltaTime ) override;
 
-	/**
-	* The widget should respond by populating the OutDrawElements array with FDrawElements
-	* that represent it and any of its children.
-	*
-	* @param AllottedGeometry  The FGeometry that describes an area in which the widget should appear.
-	* @param MyClippingRect    The clipping rectangle allocated for this widget and its children.
-	* @param OutDrawElements   A list of FDrawElements to populate with the output.
-	* @param LayerId           The Layer onto which this widget should be rendered.
-	* @param InColorAndOpacity ColorAverage and Opacity to be applied to all the descendants of the widget being painted
-	* @param bParentEnabled	True if the parent of this widget is enabled.
-	*
-	* @return The maximum layer ID attained by this widget or any of its children.
-	*/
-	virtual int32 OnPaint( const FGeometry& AllottedGeometry, const FSlateRect& MyClippingRect, FSlateWindowElementList& OutDrawElements, int32 LayerId, const FWidgetStyle& InWidgetStyle, bool bParentEnabled ) const override;
+	virtual int32 OnPaint( const FPaintArgs& Args, const FGeometry& AllottedGeometry, const FSlateRect& MyClippingRect, FSlateWindowElementList& OutDrawElements, int32 LayerId, const FWidgetStyle& InWidgetStyle, bool bParentEnabled ) const override;
 
-	/**
-	* The system calls this method to notify the widget that a mouse button was pressed within it. This event is bubbled.
-	*
-	* @param MyGeometry The Geometry of the widget receiving the event
-	* @param MouseEvent Information about the input event
-	*
-	* @return Whether the event was handled along with possible requests for the system to take action.
-	*/
 	virtual FReply OnMouseButtonDown( const FGeometry& MyGeometry, const FPointerEvent& MouseEvent ) override;
 
-	/**
-	* The system calls this method to notify the widget that a mouse button was release within it. This event is bubbled.
-	*
-	* @param MyGeometry The Geometry of the widget receiving the event
-	* @param MouseEvent Information about the input event
-	*
-	* @return Whether the event was handled along with possible requests for the system to take action.
-	*/
 	virtual FReply OnMouseButtonUp( const FGeometry& MyGeometry, const FPointerEvent& MouseEvent ) override;
 
-	/**
-	* The system calls this method to notify the widget that a mouse moved within it. This event is bubbled.
-	*
-	* @param MyGeometry The Geometry of the widget receiving the event
-	* @param MouseEvent Information about the input event
-	*
-	* @return Whether the event was handled along with possible requests for the system to take action.
-	*/
 	virtual FReply OnMouseMove( const FGeometry& MyGeometry, const FPointerEvent& MouseEvent ) override;
 
-	/**
-	* The system will use this event to notify a widget that the cursor has entered it. This event is NOT bubbled.
-	*
-	* @param MyGeometry The Geometry of the widget receiving the event
-	* @param MouseEvent Information about the input event
-	*/
 	virtual void OnMouseEnter( const FGeometry& MyGeometry, const FPointerEvent& MouseEvent ) override;
 
-	/**
-	* The system will use this event to notify a widget that the cursor has left it. This event is NOT bubbled.
-	*
-	* @param MouseEvent Information about the input event
-	*/
 	virtual void OnMouseLeave( const FPointerEvent& MouseEvent ) override;
 
-	/**
-	* Called when the mouse wheel is spun. This event is bubbled.
-	*
-	* @param  MouseEvent  Mouse event
-	*
-	* @return  Returns whether the event was handled, along with other possible actions
-	*/
 	virtual FReply OnMouseWheel( const FGeometry& MyGeometry, const FPointerEvent& MouseEvent ) override;
 
-	/**
-	* Called when a mouse button is double clicked.  Override this in derived classes.
-	*
-	* @param  InMyGeometry  Widget geometry
-	* @param  InMouseEvent  Mouse button event
-	*
-	* @return  Returns whether the event was handled, along with other possible actions
-	*/
 	virtual FReply OnMouseButtonDoubleClick( const FGeometry& MyGeometry, const FPointerEvent& MouseEvent ) override;
 
-	/**
-	* Called when the system wants to know which cursor to display for this Widget.  This event is bubbled.
-	*
-	* @return  The cursor requested (can be None.)
-	*/
 	virtual FCursorReply OnCursorQuery( const FGeometry& MyGeometry, const FPointerEvent& CursorEvent ) const override;
 
 	void ShowContextMenu( const FVector2D& ScreenSpacePosition );
diff --git a/Engine/Source/Developer/Profiler/Private/Widgets/SProfilerThreadView.cpp b/Engine/Source/Developer/Profiler/Private/Widgets/SProfilerThreadView.cpp
--- a/Engine/Source/Developer/Profiler/Private/Widgets/SProfilerThreadView.cpp
+++ b/Engine/Source/Developer/Profiler/Private/Widgets/SProfilerThreadView.cpp
@@ -80,7 +80,7 @@ void SProfilerThreadView::Tick( const FGeometry& AllottedGeometry, const double
 	}
 }
 
-int32 SProfilerThreadView::OnPaint( const FGeometry& AllottedGeometry, const FSlateRect& MyClippingRect, FSlateWindowElementList& OutDrawElements, int32 LayerId, const FWidgetStyle& InWidgetStyle, bool bParentEnabled ) const
+int32 SProfilerThreadView::OnPaint( const FPaintArgs& Args, const FGeometry& AllottedGeometry, const FSlateRect& MyClippingRect, FSlateWindowElementList& OutDrawElements, int32 LayerId, const FWidgetStyle& InWidgetStyle, bool bParentEnabled ) const
 {
 	//	SCOPE_LOG_TIME_FUNC();
 
@@ -165,7 +165,7 @@ int32 SProfilerThreadView::OnPaint( const FGeometry& AllottedGeometry, const FSl
 	// Reset paint state.
 	PaintState = nullptr;
 
-	return SCompoundWidget::OnPaint( AllottedGeometry, MyClippingRect, OutDrawElements, LayerId, InWidgetStyle, bParentEnabled && IsEnabled() );
+	return SCompoundWidget::OnPaint( Args, AllottedGeometry, MyClippingRect, OutDrawElements, LayerId, InWidgetStyle, bParentEnabled && IsEnabled() );
 }
 
 
diff --git a/Engine/Source/Developer/Profiler/Private/Widgets/SProfilerThreadView.h b/Engine/Source/Developer/Profiler/Private/Widgets/SProfilerThreadView.h
--- a/Engine/Source/Developer/Profiler/Private/Widgets/SProfilerThreadView.h
+++ b/Engine/Source/Developer/Profiler/Private/Widgets/SProfilerThreadView.h
@@ -124,99 +124,24 @@ public:
 	/** Resets internal widget's data to the default one. */
 	void Reset();
 
-	/**
-	* Ticks this widget.  Override in derived classes, but always call the parent implementation.
-	*
-	* @param  AllottedGeometry The space allotted for this widget
-	* @param  InCurrentTime  Current absolute real time
-	* @param  InDeltaTime  Real time passed since last tick
-	*/
 	virtual void Tick( const FGeometry& AllottedGeometry, const double InCurrentTime, const float InDeltaTime ) override;
 
-	/**
-	* The widget should respond by populating the OutDrawElements array with FDrawElements
-	* that represent it and any of its children.
-	*
-	* @param AllottedGeometry  The FGeometry that describes an area in which the widget should appear.
-	* @param MyClippingRect    The clipping rectangle allocated for this widget and its children.
-	* @param OutDrawElements   A list of FDrawElements to populate with the output.
-	* @param LayerId           The Layer onto which this widget should be rendered.
-	* @param InColorAndOpacity ColorAverage and Opacity to be applied to all the descendants of the widget being painted
-	* @param bParentEnabled	True if the parent of this widget is enabled.
-	*
-	* @return The maximum layer ID attained by this widget or any of its children.
-	*/
-	virtual int32 OnPaint( const FGeometry& AllottedGeometry, const FSlateRect& MyClippingRect, FSlateWindowElementList& OutDrawElements, int32 LayerId, const FWidgetStyle& InWidgetStyle, bool bParentEnabled ) const override;
+	virtual int32 OnPaint( const FPaintArgs& Args, const FGeometry& AllottedGeometry, const FSlateRect& MyClippingRect, FSlateWindowElementList& OutDrawElements, int32 LayerId, const FWidgetStyle& InWidgetStyle, bool bParentEnabled ) const override;
 
-	/**
-	* The system calls this method to notify the widget that a mouse button was pressed within it. This event is bubbled.
-	*
-	* @param MyGeometry The Geometry of the widget receiving the event
-	* @param MouseEvent Information about the input event
-	*
-	* @return Whether the event was handled along with possible requests for the system to take action.
-	*/
 	virtual FReply OnMouseButtonDown( const FGeometry& MyGeometry, const FPointerEvent& MouseEvent ) override;
 
-	/**
-	* The system calls this method to notify the widget that a mouse button was release within it. This event is bubbled.
-	*
-	* @param MyGeometry The Geometry of the widget receiving the event
-	* @param MouseEvent Information about the input event
-	*
-	* @return Whether the event was handled along with possible requests for the system to take action.
-	*/
 	virtual FReply OnMouseButtonUp( const FGeometry& MyGeometry, const FPointerEvent& MouseEvent ) override;
 
-	/**
-	* The system calls this method to notify the widget that a mouse moved within it. This event is bubbled.
-	*
-	* @param MyGeometry The Geometry of the widget receiving the event
-	* @param MouseEvent Information about the input event
-	*
-	* @return Whether the event was handled along with possible requests for the system to take action.
-	*/
 	virtual FReply OnMouseMove( const FGeometry& MyGeometry, const FPointerEvent& MouseEvent ) override;
 
-	/**
-	* The system will use this event to notify a widget that the cursor has entered it. This event is NOT bubbled.
-	*
-	* @param MyGeometry The Geometry of the widget receiving the event
-	* @param MouseEvent Information about the input event
-	*/
 	virtual void OnMouseEnter( const FGeometry& MyGeometry, const FPointerEvent& MouseEvent ) override;
 
-	/**
-	* The system will use this event to notify a widget that the cursor has left it. This event is NOT bubbled.
-	*
-	* @param MouseEvent Information about the input event
-	*/
 	virtual void OnMouseLeave( const FPointerEvent& MouseEvent ) override;
 
-	/**
-	* Called when the mouse wheel is spun. This event is bubbled.
-	*
-	* @param  MouseEvent  Mouse event
-	*
-	* @return  Returns whether the event was handled, along with other possible actions
-	*/
 	virtual FReply OnMouseWheel( const FGeometry& MyGeometry, const FPointerEvent& MouseEvent ) override;
 
-	/**
-	* Called when a mouse button is double clicked.  Override this in derived classes.
-	*
-	* @param  InMyGeometry  Widget geometry
-	* @param  InMouseEvent  Mouse button event
-	*
-	* @return  Returns whether the event was handled, along with other possible actions
-	*/
 	virtual FReply OnMouseButtonDoubleClick( const FGeometry& MyGeometry, const FPointerEvent& MouseEvent ) override;
 
-	/**
-	* Called when the system wants to know which cursor to display for this Widget.  This event is bubbled.
-	*
-	* @return  The cursor requested (can be None.)
-	*/
 	virtual FCursorReply OnCursorQuery( const FGeometry& MyGeometry, const FPointerEvent& CursorEvent ) const override;
 
 	void ShowContextMenu( const FVector2D& ScreenSpacePosition );
diff --git a/Engine/Source/Developer/RealtimeProfiler/Private/SRealtimeProfilerLineGraph.cpp b/Engine/Source/Developer/RealtimeProfiler/Private/SRealtimeProfilerLineGraph.cpp
--- a/Engine/Source/Developer/RealtimeProfiler/Private/SRealtimeProfilerLineGraph.cpp
+++ b/Engine/Source/Developer/RealtimeProfiler/Private/SRealtimeProfilerLineGraph.cpp
@@ -104,7 +104,7 @@ FVector2D SRealtimeProfilerLineGraph::GetWidgetPosition(float X, float Y, const
 }
 
 
-int32 SRealtimeProfilerLineGraph::OnPaint( const FGeometry& AllottedGeometry, const FSlateRect& MyClippingRect, FSlateWindowElementList& OutDrawElements, int32 LayerId, const FWidgetStyle& InWidgetStyle, bool bParentEnabled ) const
+int32 SRealtimeProfilerLineGraph::OnPaint( const FPaintArgs& Args, const FGeometry& AllottedGeometry, const FSlateRect& MyClippingRect, FSlateWindowElementList& OutDrawElements, int32 LayerId, const FWidgetStyle& InWidgetStyle, bool bParentEnabled ) const
 {
 	// Rendering info
 	bool bEnabled = ShouldBeEnabled( bParentEnabled );
@@ -329,7 +329,7 @@ int32 SRealtimeProfilerLineGraph::OnPaint( const FGeometry& AllottedGeometry, co
 	LayerId++;
 
 	// Paint children
-	SCompoundWidget::OnPaint(AllottedGeometry, MyClippingRect, OutDrawElements, LayerId, InWidgetStyle, bParentEnabled);
+	SCompoundWidget::OnPaint( Args, AllottedGeometry, MyClippingRect, OutDrawElements, LayerId, InWidgetStyle, bParentEnabled);
 
 	return LayerId;
 }
diff --git a/Engine/Source/Developer/RealtimeProfiler/Private/SRealtimeProfilerLineGraph.h b/Engine/Source/Developer/RealtimeProfiler/Private/SRealtimeProfilerLineGraph.h
--- a/Engine/Source/Developer/RealtimeProfiler/Private/SRealtimeProfilerLineGraph.h
+++ b/Engine/Source/Developer/RealtimeProfiler/Private/SRealtimeProfilerLineGraph.h
@@ -96,7 +96,7 @@ private:
 	// SWidget overrides
 	virtual void Tick( const FGeometry& AllottedGeometry, const double InCurrentTime, const float InDeltaTime ) override;
 	virtual FVector2D ComputeDesiredSize() const override;
-	virtual int32 OnPaint( const FGeometry& AllottedGeometry, const FSlateRect& MyClippingRect, FSlateWindowElementList& OutDrawElements, int32 LayerId, const FWidgetStyle& InWidgetStyle, bool bParentEnabled ) const override;
+	virtual int32 OnPaint( const FPaintArgs& Args, const FGeometry& AllottedGeometry, const FSlateRect& MyClippingRect, FSlateWindowElementList& OutDrawElements, int32 LayerId, const FWidgetStyle& InWidgetStyle, bool bParentEnabled ) const override;
 	virtual FReply OnMouseButtonDown( const FGeometry& MyGeometry, const FPointerEvent& MouseEvent ) override;
 	virtual FReply OnMouseButtonUp( const FGeometry& MyGeometry, const FPointerEvent& MouseEvent ) override;
 	virtual FReply OnMouseMove( const FGeometry& MyGeometry, const FPointerEvent& MouseEvent ) override;
diff --git a/Engine/Source/Developer/TaskGraph/Private/SGraphBar.cpp b/Engine/Source/Developer/TaskGraph/Private/SGraphBar.cpp
--- a/Engine/Source/Developer/TaskGraph/Private/SGraphBar.cpp
+++ b/Engine/Source/Developer/TaskGraph/Private/SGraphBar.cpp
@@ -55,7 +55,7 @@ void SGraphBar::Construct( const FArguments& InArgs )
 	TotalTime = 1.0;
 }
 
-int32 SGraphBar::OnPaint( const FGeometry& AllottedGeometry, const FSlateRect& MyClippingRect, FSlateWindowElementList& OutDrawElements, int32 LayerId, const FWidgetStyle& InWidgetStyle, bool bParentEnabled ) const
+int32 SGraphBar::OnPaint( const FPaintArgs& Args, const FGeometry& AllottedGeometry, const FSlateRect& MyClippingRect, FSlateWindowElementList& OutDrawElements, int32 LayerId, const FWidgetStyle& InWidgetStyle, bool bParentEnabled ) const
 {
 	// Used to track the layer ID we will return.
 	int32 RetLayerId = LayerId;
diff --git a/Engine/Source/Developer/TaskGraph/Private/STimeline.cpp b/Engine/Source/Developer/TaskGraph/Private/STimeline.cpp
--- a/Engine/Source/Developer/TaskGraph/Private/STimeline.cpp
+++ b/Engine/Source/Developer/TaskGraph/Private/STimeline.cpp
@@ -18,7 +18,7 @@ void STimeline::Construct( const FArguments& InArgs )
 	Offset = 0.0f;
 }
 
-int32 STimeline::OnPaint( const FGeometry& AllottedGeometry, const FSlateRect& MyClippingRect, FSlateWindowElementList& OutDrawElements, int32 LayerId, const FWidgetStyle& InWidgetStyle, bool bParentEnabled ) const
+int32 STimeline::OnPaint( const FPaintArgs& Args, const FGeometry& AllottedGeometry, const FSlateRect& MyClippingRect, FSlateWindowElementList& OutDrawElements, int32 LayerId, const FWidgetStyle& InWidgetStyle, bool bParentEnabled ) const
 {
 	// Used to track the layer ID we will return.
 	int32 RetLayerId = LayerId;
diff --git a/Engine/Source/Developer/TaskGraph/Public/SGraphBar.h b/Engine/Source/Developer/TaskGraph/Public/SGraphBar.h
--- a/Engine/Source/Developer/TaskGraph/Public/SGraphBar.h
+++ b/Engine/Source/Developer/TaskGraph/Public/SGraphBar.h
@@ -33,7 +33,7 @@ public:
 	void Construct( const FArguments& InArgs );
 
 	// SWidget interface
-	virtual int32 OnPaint( const FGeometry& AllottedGeometry, const FSlateRect& MyClippingRect, FSlateWindowElementList& OutDrawElements, int32 LayerId, const FWidgetStyle& InWidgetStyle, bool bParentEnabled ) const override;
+	virtual int32 OnPaint( const FPaintArgs& Args, const FGeometry& AllottedGeometry, const FSlateRect& MyClippingRect, FSlateWindowElementList& OutDrawElements, int32 LayerId, const FWidgetStyle& InWidgetStyle, bool bParentEnabled ) const override;
 	virtual void Tick( const FGeometry& AllottedGeometry, const double InCurrentTime, const float InDeltaTime ) override;
 	virtual FReply OnMouseButtonDown( const FGeometry& MyGeometry, const FPointerEvent& MouseEvent ) override;
 	virtual FReply OnMouseMove( const FGeometry& MyGeometry, const FPointerEvent& MouseEvent ) override;
diff --git a/Engine/Source/Developer/TaskGraph/Public/STimeline.h b/Engine/Source/Developer/TaskGraph/Public/STimeline.h
--- a/Engine/Source/Developer/TaskGraph/Public/STimeline.h
+++ b/Engine/Source/Developer/TaskGraph/Public/STimeline.h
@@ -31,7 +31,7 @@ public:
 	void Construct( const FArguments& InArgs );
 
 	// SWidget interface
-	virtual int32 OnPaint( const FGeometry& AllottedGeometry, const FSlateRect& MyClippingRect, FSlateWindowElementList& OutDrawElements, int32 LayerId, const FWidgetStyle& InWidgetStyle, bool bParentEnabled ) const override;
+	virtual int32 OnPaint( const FPaintArgs& Args,  const FGeometry& AllottedGeometry, const FSlateRect& MyClippingRect, FSlateWindowElementList& OutDrawElements, int32 LayerId, const FWidgetStyle& InWidgetStyle, bool bParentEnabled ) const override;
 	virtual void Tick( const FGeometry& AllottedGeometry, const double InCurrentTime, const float InDeltaTime ) override;
 	virtual FReply OnMouseButtonDown( const FGeometry& MyGeometry, const FPointerEvent& MouseEvent ) override;
 	virtual FReply OnMouseMove( const FGeometry& MyGeometry, const FPointerEvent& MouseEvent ) override;
diff --git a/Engine/Source/Editor/EditorWidgets/Private/SObjectNameEditableTextBox.cpp b/Engine/Source/Editor/EditorWidgets/Private/SObjectNameEditableTextBox.cpp
--- a/Engine/Source/Editor/EditorWidgets/Private/SObjectNameEditableTextBox.cpp
+++ b/Engine/Source/Editor/EditorWidgets/Private/SObjectNameEditableTextBox.cpp
@@ -53,9 +53,9 @@ void SObjectNameEditableTextBox::Tick( const FGeometry& AllottedGeometry, const
 	}
 }
 
-int32 SObjectNameEditableTextBox::OnPaint( const FGeometry& AllottedGeometry, const FSlateRect& MyClippingRect, FSlateWindowElementList& OutDrawElements, int32 LayerId, const FWidgetStyle& InWidgetStyle, bool bParentEnabled ) const
+int32 SObjectNameEditableTextBox::OnPaint( const FPaintArgs& Args, const FGeometry& AllottedGeometry, const FSlateRect& MyClippingRect, FSlateWindowElementList& OutDrawElements, int32 LayerId, const FWidgetStyle& InWidgetStyle, bool bParentEnabled ) const
 {
-	int32 StartLayer = SCompoundWidget::OnPaint( AllottedGeometry, MyClippingRect, OutDrawElements, LayerId, InWidgetStyle, bParentEnabled );
+	int32 StartLayer = SCompoundWidget::OnPaint( Args, AllottedGeometry, MyClippingRect, OutDrawElements, LayerId, InWidgetStyle, bParentEnabled );
 
 	const int32 TextLayer = 1;
 
diff --git a/Engine/Source/Editor/EditorWidgets/Private/SObjectNameEditableTextBox.h b/Engine/Source/Editor/EditorWidgets/Private/SObjectNameEditableTextBox.h
--- a/Engine/Source/Editor/EditorWidgets/Private/SObjectNameEditableTextBox.h
+++ b/Engine/Source/Editor/EditorWidgets/Private/SObjectNameEditableTextBox.h
@@ -20,29 +20,10 @@ public:
 	void Construct( const FArguments& InArgs );
 
 protected:
-	/**
-		* Ticks this widget.  Override in derived classes, but always call the parent implementation.
-		*
-		* @param  AllottedGeometry Space allotted to this widget
-		* @param  InCurrentTime  Current absolute real time
-		* @param  InDeltaTime  Real time passed since last tick
-		*/
+
 	virtual void Tick( const FGeometry& AllottedGeometry, const double InCurrentTime, const float InDeltaTime ) override;
 
-	/**
-		* The widget should respond by populating the OutDrawElements array with FDrawElements 
-		* that represent it and any of its children.
-		*
-		* @param AllottedGeometry  The FGeometry that describes an area in which the widget should appear.
-		* @param MyClippingRect    The clipping rectangle allocated for this widget and its children.
-		* @param OutDrawElements   A list of FDrawElements to populate with the output.
-		* @param LayerId           The Layer onto which this widget should be rendered.
-		* @param InColorAndOpacity Color and Opacity to be applied to all the descendants of the widget being painted
- 		* @param bParentEnabled	True if the parent of this widget is enabled.
-		*
-		* @return The maximum layer ID attained by this widget or any of its children.
-		*/
-	virtual int32 OnPaint( const FGeometry& AllottedGeometry, const FSlateRect& MyClippingRect, FSlateWindowElementList& OutDrawElements, int32 LayerId, const FWidgetStyle& InWidgetStyle, bool bParentEnabled ) const override;
+	virtual int32 OnPaint( const FPaintArgs& Args, const FGeometry& AllottedGeometry, const FSlateRect& MyClippingRect, FSlateWindowElementList& OutDrawElements, int32 LayerId, const FWidgetStyle& InWidgetStyle, bool bParentEnabled ) const override;
 
 private:
 	/** Getter for the Text attribute of the editable text inside this widget */
diff --git a/Engine/Source/Editor/EnvironmentQueryEditor/Private/STestFunctionWidget.cpp b/Engine/Source/Editor/EnvironmentQueryEditor/Private/STestFunctionWidget.cpp
--- a/Engine/Source/Editor/EnvironmentQueryEditor/Private/STestFunctionWidget.cpp
+++ b/Engine/Source/Editor/EnvironmentQueryEditor/Private/STestFunctionWidget.cpp
@@ -13,7 +13,7 @@ FVector2D STestFunctionWidget::GetWidgetPosition(float X, float Y, const FGeomet
 	return FVector2D((X*Geom.Size.X), (Geom.Size.Y - 1) - (Y*Geom.Size.Y));
 }
 
-int32 STestFunctionWidget::OnPaint(const FGeometry& AllottedGeometry, const FSlateRect& MyClippingRect, FSlateWindowElementList& OutDrawElements, int32 LayerId, const FWidgetStyle& InWidgetStyle, bool bParentEnabled) const
+int32 STestFunctionWidget::OnPaint(const FPaintArgs& Args, const FGeometry& AllottedGeometry, const FSlateRect& MyClippingRect, FSlateWindowElementList& OutDrawElements, int32 LayerId, const FWidgetStyle& InWidgetStyle, bool bParentEnabled) const
 {
 	// Rendering info
 	bool bEnabled = ShouldBeEnabled(bParentEnabled);
diff --git a/Engine/Source/Editor/GraphEditor/Private/SGraphPanel.cpp b/Engine/Source/Editor/GraphEditor/Private/SGraphPanel.cpp
--- a/Engine/Source/Editor/GraphEditor/Private/SGraphPanel.cpp
+++ b/Engine/Source/Editor/GraphEditor/Private/SGraphPanel.cpp
@@ -98,21 +98,7 @@ SGraphPanel::~SGraphPanel()
 
 //////////////////////////////////////////////////////////////////////////
 
-/**
- * The widget should respond by populating the OutDrawElements array with FDrawElements 
- * that represent it and any of its children.
- *
- * @param AllottedGeometry  The FGeometry that describes an area in which the widget should appear.
- * @param MyClippingRect    The clipping rectangle allocated for this widget and its children.
- * @param OutDrawElements   A list of FDrawElements to populate with the output.
- * @param LayerId           The Layer onto which this widget should be rendered.
- * @param InColorAndOpacity Color and Opacity to be applied to all the descendants of the widget being painted
- * @param bParentEnabled	True if the parent of this widget is enabled.
- *
- * @return The maximum layer ID attained by this widget or any of its children.
- */
-
-int32 SGraphPanel::OnPaint( const FGeometry& AllottedGeometry, const FSlateRect& MyClippingRect, FSlateWindowElementList& OutDrawElements, int32 LayerId, const FWidgetStyle& InWidgetStyle, bool bParentEnabled ) const
+int32 SGraphPanel::OnPaint( const FPaintArgs& Args, const FGeometry& AllottedGeometry, const FSlateRect& MyClippingRect, FSlateWindowElementList& OutDrawElements, int32 LayerId, const FWidgetStyle& InWidgetStyle, bool bParentEnabled ) const
 {
 #if SLATE_HD_STATS
 	SCOPE_CYCLE_COUNTER( STAT_SlateOnPaint_SGraphPanel );
@@ -274,7 +260,7 @@ int32 SGraphPanel::OnPaint( const FGeometry& AllottedGeometry, const FSlateRect&
 					const FWidgetStyle& NodeStyleToUse = (bNodeIsDifferent && !bNodeIsNotUsableInCurrentContext)? InWidgetStyle : FadedStyle;
 
 					// Draw the node.O
-					CurWidgetsMaxLayerId = CurWidget.Widget->Paint( CurWidget.Geometry, MyClippingRect, OutDrawElements, ChildLayerId, NodeStyleToUse, ShouldBeEnabled( bParentEnabled ) );
+					CurWidgetsMaxLayerId = CurWidget.Widget->Paint( Args.WithNewParent(this), CurWidget.Geometry, MyClippingRect, OutDrawElements, ChildLayerId, NodeStyleToUse, ShouldBeEnabled( bParentEnabled ) );
 				}
 
 				// Draw the node's overlay, if it has one.
@@ -324,7 +310,7 @@ int32 SGraphPanel::OnPaint( const FGeometry& AllottedGeometry, const FSlateRect&
 
 								const FGeometry WidgetGeometry = CurWidget.Geometry.MakeChild(OverlayInfo.OverlayOffset, OverlayInfo.Widget->GetDesiredSize(), 1.f);
 
-								OverlayInfo.Widget->Paint(WidgetGeometry, MyClippingRect, OutDrawElements, CurWidgetsMaxLayerId, InWidgetStyle, bParentEnabled);
+								OverlayInfo.Widget->Paint(Args.WithNewParent(this), WidgetGeometry, MyClippingRect, OutDrawElements, CurWidgetsMaxLayerId, InWidgetStyle, bParentEnabled);
 							}
 						}
 					}
diff --git a/Engine/Source/Editor/KismetWidgets/Private/SScrubWidget.cpp b/Engine/Source/Editor/KismetWidgets/Private/SScrubWidget.cpp
--- a/Engine/Source/Editor/KismetWidgets/Private/SScrubWidget.cpp
+++ b/Engine/Source/Editor/KismetWidgets/Private/SScrubWidget.cpp
@@ -54,7 +54,7 @@ void SScrubWidget::Construct( const SScrubWidget::FArguments& InArgs )
 	bAllowZoom = InArgs._bAllowZoom;
 }
 
-int32 SScrubWidget::OnPaint( const FGeometry& AllottedGeometry, const FSlateRect& MyClippingRect, FSlateWindowElementList& OutDrawElements, int32 LayerId, const FWidgetStyle& InWidgetStyle, bool bParentEnabled ) const
+int32 SScrubWidget::OnPaint( const FPaintArgs& Args, const FGeometry& AllottedGeometry, const FSlateRect& MyClippingRect, FSlateWindowElementList& OutDrawElements, int32 LayerId, const FWidgetStyle& InWidgetStyle, bool bParentEnabled ) const
 {
 	const bool bActiveFeedback = IsHovered() || bDragging;
 	const FSlateBrush* BackgroundImage = bActiveFeedback ?
@@ -185,10 +185,10 @@ int32 SScrubWidget::OnPaint( const FGeometry& AllottedGeometry, const FSlateRect
 				);
 		}
 
-		return FMath::Max( ArrowLayer, SCompoundWidget::OnPaint(AllottedGeometry, MyClippingRect, OutDrawElements, ArrowLayer, InWidgetStyle, bEnabled ) );
+		return FMath::Max( ArrowLayer, SCompoundWidget::OnPaint( Args, AllottedGeometry, MyClippingRect, OutDrawElements, ArrowLayer, InWidgetStyle, bEnabled ) );
 	}
 
-	return SCompoundWidget::OnPaint(AllottedGeometry, MyClippingRect, OutDrawElements, LayerId, InWidgetStyle, bEnabled );
+	return SCompoundWidget::OnPaint( Args, AllottedGeometry, MyClippingRect, OutDrawElements, LayerId, InWidgetStyle, bEnabled );
 }
 
 FReply SScrubWidget::OnMouseButtonDown( const FGeometry& MyGeometry, const FPointerEvent& MouseEvent )
diff --git a/Engine/Source/Editor/LandscapeEditor/Private/LandscapeEditorDetailCustomization_AlphaBrush.cpp b/Engine/Source/Editor/LandscapeEditor/Private/LandscapeEditorDetailCustomization_AlphaBrush.cpp
--- a/Engine/Source/Editor/LandscapeEditor/Private/LandscapeEditorDetailCustomization_AlphaBrush.cpp
+++ b/Engine/Source/Editor/LandscapeEditor/Private/LandscapeEditorDetailCustomization_AlphaBrush.cpp
@@ -104,7 +104,7 @@ public:
 		return FText::FromString(Texture.Get()->GetOutermost()->GetName());
 	}
 
-	virtual int32 OnPaint(const FGeometry& AllottedGeometry, const FSlateRect& MyClippingRect, FSlateWindowElementList& OutDrawElements, int32 LayerId, const FWidgetStyle& InWidgetStyle, bool bParentEnabled) const
+	virtual int32 OnPaint(const FPaintArgs& Args, const FGeometry& AllottedGeometry, const FSlateRect& MyClippingRect, FSlateWindowElementList& OutDrawElements, int32 LayerId, const FWidgetStyle& InWidgetStyle, bool bParentEnabled) const
 	{
 		if (Texture.IsBound() || TextureChannel.IsBound())
 		{
@@ -120,7 +120,7 @@ public:
 			}
 		}
 
-		return SCompoundWidget::OnPaint(AllottedGeometry, MyClippingRect, OutDrawElements, LayerId, InWidgetStyle, bParentEnabled);
+		return SCompoundWidget::OnPaint(Args, AllottedGeometry, MyClippingRect, OutDrawElements, LayerId, InWidgetStyle, bParentEnabled);
 	}
 
 private:
diff --git a/Engine/Source/Editor/NewLevelDialog/Private/NewLevelDialogModule.cpp b/Engine/Source/Editor/NewLevelDialog/Private/NewLevelDialogModule.cpp
--- a/Engine/Source/Editor/NewLevelDialog/Private/NewLevelDialogModule.cpp
+++ b/Engine/Source/Editor/NewLevelDialog/Private/NewLevelDialogModule.cpp
@@ -29,7 +29,7 @@ public:
 			});
 	}
 
-	virtual int32 OnPaint( const FGeometry& AllottedGeometry, const FSlateRect& MyClippingRect, FSlateWindowElementList& OutDrawElements, int32 LayerId, const FWidgetStyle& InWidgetStyle, bool bParentEnabled ) const
+	virtual int32 OnPaint( const FPaintArgs& Args, const FGeometry& AllottedGeometry, const FSlateRect& MyClippingRect, FSlateWindowElementList& OutDrawElements, int32 LayerId, const FWidgetStyle& InWidgetStyle, bool bParentEnabled ) const
 	{
 		bool bEnableGammaCorrection = true;
 		bool bAllowBlending = false;
diff --git a/Engine/Source/Editor/Persona/Private/SAnimCurvePanel.cpp b/Engine/Source/Editor/Persona/Private/SAnimCurvePanel.cpp
--- a/Engine/Source/Editor/Persona/Private/SAnimCurvePanel.cpp
+++ b/Engine/Source/Editor/Persona/Private/SAnimCurvePanel.cpp
@@ -112,7 +112,7 @@ public:
 
 protected:
 	// SWidget interface
-	virtual int32 OnPaint(const FGeometry& AllottedGeometry, const FSlateRect& MyClippingRect, FSlateWindowElementList& OutDrawElements, int32 LayerId, const FWidgetStyle& InWidgetStyle, bool bParentEnabled) const override;
+	virtual int32 OnPaint(const FPaintArgs& Args, const FGeometry& AllottedGeometry, const FSlateRect& MyClippingRect, FSlateWindowElementList& OutDrawElements, int32 LayerId, const FWidgetStyle& InWidgetStyle, bool bParentEnabled) const override;
 	// SWidget interface
 
 	// SCurveEditor interface
@@ -153,9 +153,9 @@ float SAnimCurveEd::GetTimeStep(FTrackScaleInfo &ScaleInfo)const
 	return 0.0f;
 }
 
-int32 SAnimCurveEd::OnPaint(const FGeometry& AllottedGeometry, const FSlateRect& MyClippingRect, FSlateWindowElementList& OutDrawElements, int32 LayerId, const FWidgetStyle& InWidgetStyle, bool bParentEnabled) const
+int32 SAnimCurveEd::OnPaint(const FPaintArgs& Args, const FGeometry& AllottedGeometry, const FSlateRect& MyClippingRect, FSlateWindowElementList& OutDrawElements, int32 LayerId, const FWidgetStyle& InWidgetStyle, bool bParentEnabled) const
 {
-	int32 NewLayerId = SCurveEditor::OnPaint(AllottedGeometry, MyClippingRect, OutDrawElements, LayerId, InWidgetStyle, bParentEnabled) + 1;
+	int32 NewLayerId = SCurveEditor::OnPaint(Args, AllottedGeometry, MyClippingRect, OutDrawElements, LayerId, InWidgetStyle, bParentEnabled) + 1;
 
 	float Value = 0.f;
 
diff --git a/Engine/Source/Editor/Persona/Private/SAnimMontagePanel.cpp b/Engine/Source/Editor/Persona/Private/SAnimMontagePanel.cpp
--- a/Engine/Source/Editor/Persona/Private/SAnimMontagePanel.cpp
+++ b/Engine/Source/Editor/Persona/Private/SAnimMontagePanel.cpp
@@ -68,7 +68,7 @@ public:
 	SLATE_END_ARGS()
 
 	void Construct(const FArguments& InArgs);
-	virtual int32 OnPaint(const FGeometry& AllottedGeometry, const FSlateRect& MyClippingRect, FSlateWindowElementList& OutDrawElements, int32 LayerId, const FWidgetStyle& InWidgetStyle, bool bParentEnabled) const override;
+	virtual int32 OnPaint(const FPaintArgs& Args, const FGeometry& AllottedGeometry, const FSlateRect& MyClippingRect, FSlateWindowElementList& OutDrawElements, int32 LayerId, const FWidgetStyle& InWidgetStyle, bool bParentEnabled) const override;
 
 	virtual FVector2D GetDragDropScreenSpacePosition(const FGeometry& ParentAllottedGeometry, const FDragDropEvent& DragDropEvent) const override;
 
@@ -184,7 +184,7 @@ void SMontageBranchingPointNode::OnSnapNodeDataPosition(float OriginalX, float S
 	BranchingPoint->TriggerTimeOffset = GetTriggerTimeOffsetForType(Offset);
 }
 
-int32 SMontageBranchingPointNode::OnPaint(const FGeometry& AllottedGeometry, const FSlateRect& MyClippingRect, FSlateWindowElementList& OutDrawElements, int32 LayerId, const FWidgetStyle& InWidgetStyle, bool bParentEnabled) const
+int32 SMontageBranchingPointNode::OnPaint(const FPaintArgs& Args, const FGeometry& AllottedGeometry, const FSlateRect& MyClippingRect, FSlateWindowElementList& OutDrawElements, int32 LayerId, const FWidgetStyle& InWidgetStyle, bool bParentEnabled) const
 {
 	int32 TextLayerID = LayerId+1;
 
diff --git a/Engine/Source/Editor/Persona/Private/SAnimNotifyPanel.cpp b/Engine/Source/Editor/Persona/Private/SAnimNotifyPanel.cpp
--- a/Engine/Source/Editor/Persona/Private/SAnimNotifyPanel.cpp
+++ b/Engine/Source/Editor/Persona/Private/SAnimNotifyPanel.cpp
@@ -85,7 +85,7 @@ public:
 	// End of SNodePanel::SNode
 
 	virtual FVector2D ComputeDesiredSize() const  override;
-	virtual int32 OnPaint(const FGeometry& AllottedGeometry, const FSlateRect& MyClippingRect, FSlateWindowElementList& OutDrawElements, int32 LayerId, const FWidgetStyle& InWidgetStyle, bool bParentEnabled) const override;
+	virtual int32 OnPaint(const FPaintArgs& Args, const FGeometry& AllottedGeometry, const FSlateRect& MyClippingRect, FSlateWindowElementList& OutDrawElements, int32 LayerId, const FWidgetStyle& InWidgetStyle, bool bParentEnabled) const override;
 
 	/** Helpers to draw scrub handles and snap offsets */
 	void DrawHandleOffset( const float& Offset, const float& HandleCentre, FSlateWindowElementList& OutDrawElements, int32 MarkerLayer, const FGeometry &AllottedGeometry, const FSlateRect& MyClippingRect ) const;
@@ -229,7 +229,7 @@ public:
 
 	// SWidget interface
 	virtual void Tick( const FGeometry& AllottedGeometry, const double InCurrentTime, const float InDeltaTime ) {UpdateCachedGeometry(AllottedGeometry);}
-	virtual int32 OnPaint(const FGeometry& AllottedGeometry, const FSlateRect& MyClippingRect, FSlateWindowElementList& OutDrawElements, int32 LayerId, const FWidgetStyle& InWidgetStyle, bool bParentEnabled) const override;
+	virtual int32 OnPaint(const FPaintArgs& Args, const FGeometry& AllottedGeometry, const FSlateRect& MyClippingRect, FSlateWindowElementList& OutDrawElements, int32 LayerId, const FWidgetStyle& InWidgetStyle, bool bParentEnabled) const override;
 	// End of SWidget interface
 
 	/** Returns the cached rendering geometry of this track */
@@ -1023,7 +1023,7 @@ FVector2D SAnimNotifyNode::GetSize() const
 	return WidgetSize;
 }
 
-int32 SAnimNotifyNode::OnPaint(const FGeometry& AllottedGeometry, const FSlateRect& MyClippingRect, FSlateWindowElementList& OutDrawElements, int32 LayerId, const FWidgetStyle& InWidgetStyle, bool bParentEnabled) const
+int32 SAnimNotifyNode::OnPaint(const FPaintArgs& Args, const FGeometry& AllottedGeometry, const FSlateRect& MyClippingRect, FSlateWindowElementList& OutDrawElements, int32 LayerId, const FWidgetStyle& InWidgetStyle, bool bParentEnabled) const
 {
 	int32 MarkerLayer = LayerId + 1;
 	int32 ScrubHandleID = MarkerLayer + 1;
@@ -1422,7 +1422,7 @@ FVector2D SAnimNotifyTrack::ComputeDesiredSize() const
 	return Size;
 }
 
-int32 SAnimNotifyTrack::OnPaint(const FGeometry& AllottedGeometry, const FSlateRect& MyClippingRect, FSlateWindowElementList& OutDrawElements, int32 LayerId, const FWidgetStyle& InWidgetStyle, bool bParentEnabled) const
+int32 SAnimNotifyTrack::OnPaint(const FPaintArgs& Args, const FGeometry& AllottedGeometry, const FSlateRect& MyClippingRect, FSlateWindowElementList& OutDrawElements, int32 LayerId, const FWidgetStyle& InWidgetStyle, bool bParentEnabled) const
 {
 	const FSlateBrush* StyleInfo = FEditorStyle::GetBrush( TEXT( "Persona.NotifyEditor.NotifyTrackBackground" ) );
 	FLinearColor Color = TrackColor.Get();
@@ -1558,7 +1558,7 @@ int32 SAnimNotifyTrack::OnPaint(const FGeometry& AllottedGeometry, const FSlateR
 			);
 	}
 
-	return SCompoundWidget::OnPaint(AllottedGeometry, MyClippingRect, OutDrawElements, CustomLayerId, InWidgetStyle, bParentEnabled);
+	return SCompoundWidget::OnPaint(Args, AllottedGeometry, MyClippingRect, OutDrawElements, CustomLayerId, InWidgetStyle, bParentEnabled);
 }
 
 void SAnimNotifyTrack::FillNewNotifyStateMenu(FMenuBuilder& MenuBuilder)
diff --git a/Engine/Source/Editor/Persona/Private/SAnimationBlendSpace.cpp b/Engine/Source/Editor/Persona/Private/SAnimationBlendSpace.cpp
--- a/Engine/Source/Editor/Persona/Private/SAnimationBlendSpace.cpp
+++ b/Engine/Source/Editor/Persona/Private/SAnimationBlendSpace.cpp
@@ -846,9 +846,9 @@ FText SBlendSpaceGridWidget::GetInputText(const FVector & GridPos) const
 	return FText::Format( LOCTEXT("BlendSpaceParamGridMessage", "{BlendParam01DisplayName} [{XGridPos}]\n{BlendParam02DisplayName} [{YGridPos}]\n"), Args );
 }
 
-int32 SBlendSpaceGridWidget::OnPaint( const FGeometry& AllottedGeometry, const FSlateRect& MyClippingRect, FSlateWindowElementList& OutDrawElements, int32 LayerId, const FWidgetStyle& InWidgetStyle, bool bParentEnabled ) const
+int32 SBlendSpaceGridWidget::OnPaint( const FPaintArgs& Args, const FGeometry& AllottedGeometry, const FSlateRect& MyClippingRect, FSlateWindowElementList& OutDrawElements, int32 LayerId, const FWidgetStyle& InWidgetStyle, bool bParentEnabled ) const
 {
-	SCompoundWidget::OnPaint(AllottedGeometry, MyClippingRect, OutDrawElements, LayerId, InWidgetStyle, bParentEnabled && IsEnabled() );
+	SCompoundWidget::OnPaint( Args, AllottedGeometry, MyClippingRect, OutDrawElements, LayerId, InWidgetStyle, bParentEnabled && IsEnabled() );
 
 	FDrawLines LineBatcher;
 
diff --git a/Engine/Source/Editor/Persona/Private/SAnimationBlendSpace.h b/Engine/Source/Editor/Persona/Private/SAnimationBlendSpace.h
--- a/Engine/Source/Editor/Persona/Private/SAnimationBlendSpace.h
+++ b/Engine/Source/Editor/Persona/Private/SAnimationBlendSpace.h
@@ -468,20 +468,7 @@ public:
 	 */
 	void Construct(const FArguments& InArgs);
 
-	/**
-	 * The widget should respond by populating the OutDrawElements array with FDrawElements 
-	 * that represent it and any of its children.
-	 *
-	 * @param AllottedGeometry  The FGeometry that describes an area in which the widget should appear.
-	 * @param MyClippingRect    The clipping rectangle allocated for this widget and its children.
-	 * @param OutDrawElements   A list of FDrawElements to populate with the output.
-	 * @param LayerId           The Layer onto which this widget should be rendered.
-	 * @param InColorAndOpacity Color and Opacity to be applied to all the descendants of the widget being painted
- 	 * @param bParentEnabled	True if the parent of this widget is enabled.
-	 *
-	 * @return The maximum layer ID attained by this widget or any of its children.
-	 */
-	virtual int32 OnPaint( const FGeometry& AllottedGeometry, const FSlateRect& MyClippingRect, FSlateWindowElementList& OutDrawElements, int32 LayerId, const FWidgetStyle& InWidgetStyle, bool bParentEnabled ) const override;
+	virtual int32 OnPaint( const FPaintArgs& Args, const FGeometry& AllottedGeometry, const FSlateRect& MyClippingRect, FSlateWindowElementList& OutDrawElements, int32 LayerId, const FWidgetStyle& InWidgetStyle, bool bParentEnabled ) const override;
 
 	/**
 	 * Clear all points
diff --git a/Engine/Source/Editor/Persona/Private/SAnimationBlendSpace1D.cpp b/Engine/Source/Editor/Persona/Private/SAnimationBlendSpace1D.cpp
--- a/Engine/Source/Editor/Persona/Private/SAnimationBlendSpace1D.cpp
+++ b/Engine/Source/Editor/Persona/Private/SAnimationBlendSpace1D.cpp
@@ -276,9 +276,9 @@ FReply SBlendSpace1DWidget::UpdateLastMousePosition( const FGeometry& MyGeometry
 	return FReply::Handled();
 }
 
-int SBlendSpace1DWidget::OnPaint( const FGeometry& AllottedGeometry, const FSlateRect& MyClippingRect, FSlateWindowElementList& OutDrawElements, int32 LayerId, const FWidgetStyle& InWidgetStyle, bool bParentEnabled ) const
+int SBlendSpace1DWidget::OnPaint( const FPaintArgs& Args, const FGeometry& AllottedGeometry, const FSlateRect& MyClippingRect, FSlateWindowElementList& OutDrawElements, int32 LayerId, const FWidgetStyle& InWidgetStyle, bool bParentEnabled ) const
 {
-	SCompoundWidget::OnPaint(AllottedGeometry, MyClippingRect, OutDrawElements, LayerId, InWidgetStyle, bParentEnabled && IsEnabled() );
+	SCompoundWidget::OnPaint(Args, AllottedGeometry, MyClippingRect, OutDrawElements, LayerId, InWidgetStyle, bParentEnabled && IsEnabled() );
 
 	FDrawLines LineBatcher;
 
diff --git a/Engine/Source/Editor/Persona/Private/SAnimationBlendSpace1D.h b/Engine/Source/Editor/Persona/Private/SAnimationBlendSpace1D.h
--- a/Engine/Source/Editor/Persona/Private/SAnimationBlendSpace1D.h
+++ b/Engine/Source/Editor/Persona/Private/SAnimationBlendSpace1D.h
@@ -109,20 +109,7 @@ public:
 	 */
 	 void Construct(const FArguments& InArgs);
 
-	/**
-	 * The widget should respond by populating the OutDrawElements array with FDrawElements 
-	 * that represent it and any of its children.
-	 *
-	 * @param AllottedGeometry  The FGeometry that describes an area in which the widget should appear.
-	 * @param MyClippingRect    The clipping rectangle allocated for this widget and its children.
-	 * @param OutDrawElements   A list of FDrawElements to populate with the output.
-	 * @param LayerId           The Layer onto which this widget should be rendered.
-	 * @param InColorAndOpacity Color and Opacity to be applied to all the descendants of the widget being painted
- 	 * @param bParentEnabled	True if the parent of this widget is enabled.
-	 *
-	 * @return The maximum layer ID attained by this widget or any of its children.
-	 */
-	virtual int32 OnPaint( const FGeometry& AllottedGeometry, const FSlateRect& MyClippingRect, FSlateWindowElementList& OutDrawElements, int32 LayerId, const FWidgetStyle& InWidgetStyle, bool bParentEnabled ) const override;
+	virtual int32 OnPaint( const FPaintArgs& Args, const FGeometry& AllottedGeometry, const FSlateRect& MyClippingRect, FSlateWindowElementList& OutDrawElements, int32 LayerId, const FWidgetStyle& InWidgetStyle, bool bParentEnabled ) const override;
 
 	/**
 	 * A Panel's desired size in the space required to arrange of its children on the screen while respecting all of
diff --git a/Engine/Source/Editor/Persona/Private/SRetargetSourceManager.cpp b/Engine/Source/Editor/Persona/Private/SRetargetSourceManager.cpp
--- a/Engine/Source/Editor/Persona/Private/SRetargetSourceManager.cpp
+++ b/Engine/Source/Editor/Persona/Private/SRetargetSourceManager.cpp
@@ -440,9 +440,9 @@ void SRetargetSourceManager::CreateRetargetSourceList( const FString& SearchText
 	RetargetSourceListView->RequestListRefresh();
 }
 
-int32 SRetargetSourceManager::OnPaint( const FGeometry& AllottedGeometry, const FSlateRect& MyClippingRect, FSlateWindowElementList& OutDrawElements, int32 LayerId, const FWidgetStyle& InWidgetStyle, bool bParentEnabled ) const
+int32 SRetargetSourceManager::OnPaint( const FPaintArgs& Args, const FGeometry& AllottedGeometry, const FSlateRect& MyClippingRect, FSlateWindowElementList& OutDrawElements, int32 LayerId, const FWidgetStyle& InWidgetStyle, bool bParentEnabled ) const
 {
-	int32 Result = SCompoundWidget::OnPaint( AllottedGeometry, MyClippingRect, OutDrawElements, LayerId, InWidgetStyle, bParentEnabled );
+	const int32 Result = SCompoundWidget::OnPaint( Args, AllottedGeometry, MyClippingRect, OutDrawElements, LayerId, InWidgetStyle, bParentEnabled );
 
 	// I need to do this after first paint	
 	if (ItemToRename.IsValid())
diff --git a/Engine/Source/Editor/Persona/Private/STrack.cpp b/Engine/Source/Editor/Persona/Private/STrack.cpp
--- a/Engine/Source/Editor/Persona/Private/STrack.cpp
+++ b/Engine/Source/Editor/Persona/Private/STrack.cpp
@@ -92,7 +92,7 @@ void STrackNode::Construct(const FArguments& InArgs)
 	Font = FSlateFontInfo( FPaths::EngineContentDir() / TEXT("Slate/Fonts/Roboto-Regular.ttf"), 10 );
 }
 
-int32 STrackNode::OnPaint(const FGeometry& AllottedGeometry, const FSlateRect& MyClippingRect, FSlateWindowElementList& OutDrawElements, int32 LayerId, const FWidgetStyle& InWidgetStyle, bool bParentEnabled) const
+int32 STrackNode::OnPaint(const FPaintArgs& Args, const FGeometry& AllottedGeometry, const FSlateRect& MyClippingRect, FSlateWindowElementList& OutDrawElements, int32 LayerId, const FWidgetStyle& InWidgetStyle, bool bParentEnabled) const
 {
 	// We draw ourselves on all of the AllottedGeometry. The parent STrack should have arranged us
 	FPaintGeometry MyGeometry =	AllottedGeometry.ToPaintGeometry();
@@ -380,7 +380,7 @@ FChildren* STrack::GetChildren()
 	return &TrackNodes;
 }
 
-int32 STrack::OnPaint( const FGeometry& AllottedGeometry, const FSlateRect& MyClippingRect, FSlateWindowElementList& OutDrawElements, int32 LayerId, const FWidgetStyle& InWidgetStyle, bool bParentEnabled ) const
+int32 STrack::OnPaint( const FPaintArgs& Args, const FGeometry& AllottedGeometry, const FSlateRect& MyClippingRect, FSlateWindowElementList& OutDrawElements, int32 LayerId, const FWidgetStyle& InWidgetStyle, bool bParentEnabled ) const
 {
 	int32 CustomLayerId = LayerId + 1;
 	FPaintGeometry MyGeometry = AllottedGeometry.ToPaintGeometry();
@@ -506,7 +506,7 @@ int32 STrack::OnPaint( const FGeometry& AllottedGeometry, const FSlateRect& MyCl
 	}
 
 	
-	return SPanel::OnPaint( AllottedGeometry, MyClippingRect, OutDrawElements, CustomLayerId, InWidgetStyle, bParentEnabled );
+	return SPanel::OnPaint( Args, AllottedGeometry, MyClippingRect, OutDrawElements, CustomLayerId, InWidgetStyle, bParentEnabled );
 }
 
 // drag drop relationship
diff --git a/Engine/Source/Editor/Persona/Private/STrack.h b/Engine/Source/Editor/Persona/Private/STrack.h
--- a/Engine/Source/Editor/Persona/Private/STrack.h
+++ b/Engine/Source/Editor/Persona/Private/STrack.h
@@ -117,7 +117,7 @@ public:
 	virtual FVector2D GetOffsetRelativeToParent(const FGeometry& ParentAllottedGeometry) const;
 	virtual FVector2D GetSizeRelativeToParent(const FGeometry& ParentAllottedGeometry) const;
 	
-	virtual int32 OnPaint(const FGeometry& AllottedGeometry, const FSlateRect& MyClippingRect, FSlateWindowElementList& OutDrawElements, int32 LayerId, const FWidgetStyle& InWidgetStyle, bool bParentEnabled) const;
+	virtual int32 OnPaint(const FPaintArgs& Args, const FGeometry& AllottedGeometry, const FSlateRect& MyClippingRect, FSlateWindowElementList& OutDrawElements, int32 LayerId, const FWidgetStyle& InWidgetStyle, bool bParentEnabled) const;
 
 	// drag drop relationship
 	virtual FReply OnDragDetected( const FGeometry& MyGeometry, const FPointerEvent& MouseEvent ) override;
@@ -236,7 +236,7 @@ public:
 	void			Construct( const FArguments& InArgs );
 
 	virtual void	OnArrangeChildren( const FGeometry& AllottedGeometry, FArrangedChildren& ArrangedChildren ) const;
-	virtual int32	OnPaint( const FGeometry& AllottedGeometry, const FSlateRect& MyClippingRect, FSlateWindowElementList& OutDrawElements, int32 LayerId, const FWidgetStyle& InWidgetStyle, bool bParentEnabled ) const;
+	virtual int32	OnPaint( const FPaintArgs& Args, const FGeometry& AllottedGeometry, const FSlateRect& MyClippingRect, FSlateWindowElementList& OutDrawElements, int32 LayerId, const FWidgetStyle& InWidgetStyle, bool bParentEnabled ) const;
 
 	virtual FReply	OnDrop( const FGeometry& MyGeometry, const FDragDropEvent& DragDropEvent ) override;	
 	virtual FReply	OnDragOver( const FGeometry& MyGeometry, const FDragDropEvent& DragDropEvent ) override;
diff --git a/Engine/Source/Editor/PropertyEditor/Private/SStructureDetailsView.cpp b/Engine/Source/Editor/PropertyEditor/Private/SStructureDetailsView.cpp
--- a/Engine/Source/Editor/PropertyEditor/Private/SStructureDetailsView.cpp
+++ b/Engine/Source/Editor/PropertyEditor/Private/SStructureDetailsView.cpp
@@ -204,7 +204,7 @@ void SStructureDetailsView::SetStructureData(TSharedPtr<FStructOnScope> InStruct
 	UpdatePropertyMap();
 }
 
-int32 SStructureDetailsView::OnPaint(FGeometry const& AllottedGeometry, FSlateRect const& MyClippingRect, FSlateWindowElementList& OutDrawElements, int32 LayerId, FWidgetStyle const& InWidgetStyle, bool bParentEnabled) const
+int32 SStructureDetailsView::OnPaint(const FPaintArgs& Args, FGeometry const& AllottedGeometry, FSlateRect const& MyClippingRect, FSlateWindowElementList& OutDrawElements, int32 LayerId, FWidgetStyle const& InWidgetStyle, bool bParentEnabled) const
 {
 	int32 LayerIdOut = LayerId;
 	// if the details table is still in need of a refresh, then wait for it to
@@ -215,7 +215,7 @@ int32 SStructureDetailsView::OnPaint(FGeometry const& AllottedGeometry, FSlateRe
 	//        better, lower level solution (this could happen for other detail panels)
 	if (!DetailTree->IsPendingRefresh())
 	{
-		LayerIdOut = SDetailsViewBase::OnPaint(AllottedGeometry, MyClippingRect, OutDrawElements, LayerId, InWidgetStyle, bParentEnabled);
+		LayerIdOut = SDetailsViewBase::OnPaint(Args, AllottedGeometry, MyClippingRect, OutDrawElements, LayerId, InWidgetStyle, bParentEnabled);
 	}
 	
 	return LayerIdOut;
diff --git a/Engine/Source/Editor/PropertyEditor/Private/SStructureDetailsView.h b/Engine/Source/Editor/PropertyEditor/Private/SStructureDetailsView.h
--- a/Engine/Source/Editor/PropertyEditor/Private/SStructureDetailsView.h
+++ b/Engine/Source/Editor/PropertyEditor/Private/SStructureDetailsView.h
@@ -88,7 +88,7 @@ public:
 	virtual TSharedPtr<class FComplexPropertyNode> GetRootNode() override;
 
 	// Begin SWidget interface
-	virtual int32 OnPaint(FGeometry const& AllottedGeometry, FSlateRect const& MyClippingRect, FSlateWindowElementList& OutDrawElements, int32 LayerId, FWidgetStyle const& InWidgetStyle, bool bParentEnabled) const override;
+	virtual int32 OnPaint(const FPaintArgs& Args, FGeometry const& AllottedGeometry, FSlateRect const& MyClippingRect, FSlateWindowElementList& OutDrawElements, int32 LayerId, FWidgetStyle const& InWidgetStyle, bool bParentEnabled) const override;
 	// End SWidget interface
 protected:
 	virtual void CustomUpdatePropertyMap() override;
diff --git a/Engine/Source/Editor/PropertyEditor/Private/UserInterface/PropertyTable/SPropertyTableCell.cpp b/Engine/Source/Editor/PropertyEditor/Private/UserInterface/PropertyTable/SPropertyTableCell.cpp
--- a/Engine/Source/Editor/PropertyEditor/Private/UserInterface/PropertyTable/SPropertyTableCell.cpp
+++ b/Engine/Source/Editor/PropertyEditor/Private/UserInterface/PropertyTable/SPropertyTableCell.cpp
@@ -137,7 +137,7 @@ void SPropertyTableCell::Tick( const FGeometry& AllottedGeometry, const double I
 	}
 }
 
-int32 SPropertyTableCell::OnPaint( const FGeometry& AllottedGeometry, const FSlateRect& MyClippingRect, FSlateWindowElementList& OutDrawElements, int32 LayerId, const FWidgetStyle& InWidgetStyle, bool bParentEnabled ) const
+int32 SPropertyTableCell::OnPaint( const FPaintArgs& Args, const FGeometry& AllottedGeometry, const FSlateRect& MyClippingRect, FSlateWindowElementList& OutDrawElements, int32 LayerId, const FWidgetStyle& InWidgetStyle, bool bParentEnabled ) const
 {
 	if ( CellBackground && CellBackground->DrawAs != ESlateBrushDrawType::NoDrawType )
 	{
@@ -163,7 +163,7 @@ int32 SPropertyTableCell::OnPaint( const FGeometry& AllottedGeometry, const FSla
 			);
 	}
 
-	return SCompoundWidget::OnPaint( AllottedGeometry, MyClippingRect, OutDrawElements, LayerId, InWidgetStyle, bParentEnabled );
+	return SCompoundWidget::OnPaint( Args, AllottedGeometry, MyClippingRect, OutDrawElements, LayerId, InWidgetStyle, bParentEnabled );
 }
 
 FReply SPropertyTableCell::OnMouseButtonDown( const FGeometry& MyGeometry, const FPointerEvent& MouseEvent )
diff --git a/Engine/Source/Editor/SceneOutliner/Private/SSceneOutlinerItemWidget.cpp b/Engine/Source/Editor/SceneOutliner/Private/SSceneOutlinerItemWidget.cpp
--- a/Engine/Source/Editor/SceneOutliner/Private/SSceneOutlinerItemWidget.cpp
+++ b/Engine/Source/Editor/SceneOutliner/Private/SSceneOutlinerItemWidget.cpp
@@ -60,9 +60,9 @@ namespace SceneOutliner
 		}
 	}
 
-	int32 SSceneOutlinerItemWidget::OnPaint( const FGeometry& AllottedGeometry, const FSlateRect& MyClippingRect, FSlateWindowElementList& OutDrawElements, int32 LayerId, const FWidgetStyle& InWidgetStyle, bool bParentEnabled ) const
+	int32 SSceneOutlinerItemWidget::OnPaint( const FPaintArgs& Args, const FGeometry& AllottedGeometry, const FSlateRect& MyClippingRect, FSlateWindowElementList& OutDrawElements, int32 LayerId, const FWidgetStyle& InWidgetStyle, bool bParentEnabled ) const
 	{
-		int32 StartLayer = SCompoundWidget::OnPaint( AllottedGeometry, MyClippingRect, OutDrawElements, LayerId, InWidgetStyle, bParentEnabled );
+		int32 StartLayer = SCompoundWidget::OnPaint( Args, AllottedGeometry, MyClippingRect, OutDrawElements, LayerId, InWidgetStyle, bParentEnabled );
 
 		const int32 TextLayer = 1;	
 
diff --git a/Engine/Source/Editor/SceneOutliner/Private/SSceneOutlinerItemWidget.h b/Engine/Source/Editor/SceneOutliner/Private/SSceneOutlinerItemWidget.h
--- a/Engine/Source/Editor/SceneOutliner/Private/SSceneOutlinerItemWidget.h
+++ b/Engine/Source/Editor/SceneOutliner/Private/SSceneOutlinerItemWidget.h
@@ -25,29 +25,10 @@ namespace SceneOutliner
 		void FlashHighlight();
 
 	protected:
-		/**
-		 * Ticks this widget.  Override in derived classes, but always call the parent implementation.
-		 *
-		 * @param  AllottedGeometry Space allotted to this widget
-		 * @param  InCurrentTime  Current absolute real time
-		 * @param  InDeltaTime  Real time passed since last tick
-		 */
+
 		virtual void Tick( const FGeometry& AllottedGeometry, const double InCurrentTime, const float InDeltaTime ) override;
 
-		/**
-		 * The widget should respond by populating the OutDrawElements array with FDrawElements 
-		 * that represent it and any of its children.
-		 *
-		 * @param AllottedGeometry  The FGeometry that describes an area in which the widget should appear.
-		 * @param MyClippingRect    The clipping rectangle allocated for this widget and its children.
-		 * @param OutDrawElements   A list of FDrawElements to populate with the output.
-		 * @param LayerId           The Layer onto which this widget should be rendered.
-		 * @param InColorAndOpacity Color and Opacity to be applied to all the descendants of the widget being painted
- 		 * @param bParentEnabled	True if the parent of this widget is enabled.
-		 *
-		 * @return The maximum layer ID attained by this widget or any of its children.
-		 */
-		virtual int32 OnPaint( const FGeometry& AllottedGeometry, const FSlateRect& MyClippingRect, FSlateWindowElementList& OutDrawElements, int32 LayerId, const FWidgetStyle& InWidgetStyle, bool bParentEnabled ) const override;
+		virtual int32 OnPaint( const FPaintArgs& Args, const FGeometry& AllottedGeometry, const FSlateRect& MyClippingRect, FSlateWindowElementList& OutDrawElements, int32 LayerId, const FWidgetStyle& InWidgetStyle, bool bParentEnabled ) const override;
 
 	private:
 		/** How many pixels to extend the highlight rectangle's left side horizontally */
diff --git a/Engine/Source/Editor/Sequencer/Private/SSection.cpp b/Engine/Source/Editor/Sequencer/Private/SSection.cpp
--- a/Engine/Source/Editor/Sequencer/Private/SSection.cpp
+++ b/Engine/Source/Editor/Sequencer/Private/SSection.cpp
@@ -210,9 +210,9 @@ void SSection::CreateDragOperation( const FGeometry& MyGeometry, const FPointerE
 	
 }
 
-int32 SSection::OnPaint( const FGeometry& AllottedGeometry, const FSlateRect& MyClippingRect, FSlateWindowElementList& OutDrawElements, int32 LayerId, const FWidgetStyle& InWidgetStyle, bool bParentEnabled ) const
+int32 SSection::OnPaint( const FPaintArgs& Args, const FGeometry& AllottedGeometry, const FSlateRect& MyClippingRect, FSlateWindowElementList& OutDrawElements, int32 LayerId, const FWidgetStyle& InWidgetStyle, bool bParentEnabled ) const
 {
-	int32 StartLayer = SCompoundWidget::OnPaint( AllottedGeometry, MyClippingRect, OutDrawElements, LayerId, InWidgetStyle, bParentEnabled );
+	int32 StartLayer = SCompoundWidget::OnPaint( Args, AllottedGeometry, MyClippingRect, OutDrawElements, LayerId, InWidgetStyle, bParentEnabled );
 
 	UMovieSceneSection* SectionObject = SectionInterface->GetSectionObject();
 
diff --git a/Engine/Source/Editor/Sequencer/Private/SSequencer.cpp b/Engine/Source/Editor/Sequencer/Private/SSequencer.cpp
--- a/Engine/Source/Editor/Sequencer/Private/SSequencer.cpp
+++ b/Engine/Source/Editor/Sequencer/Private/SSequencer.cpp
@@ -60,7 +60,7 @@ public:
 		}
 	}
 	
-	virtual int32 OnPaint( const FGeometry& AllottedGeometry, const FSlateRect& MyClippingRect, FSlateWindowElementList& OutDrawElements, int32 LayerId, const FWidgetStyle& InWidgetStyle, bool bParentEnabled ) const override
+	virtual int32 OnPaint( const FPaintArgs& Args, const FGeometry& AllottedGeometry, const FSlateRect& MyClippingRect, FSlateWindowElementList& OutDrawElements, int32 LayerId, const FWidgetStyle& InWidgetStyle, bool bParentEnabled ) const override
 	{
 		float Alpha = Sequencer.Pin()->GetOverlayFadeCurve();
 		
diff --git a/Engine/Source/Editor/Sequencer/Private/SSequencerSectionAreaView.cpp b/Engine/Source/Editor/Sequencer/Private/SSequencerSectionAreaView.cpp
--- a/Engine/Source/Editor/Sequencer/Private/SSequencerSectionAreaView.cpp
+++ b/Engine/Source/Editor/Sequencer/Private/SSequencerSectionAreaView.cpp
@@ -106,7 +106,7 @@ EVisibility SSequencerSectionAreaView::GetNodeVisibility() const
 
 
 /** SWidget Interface */
-int32 SSequencerSectionAreaView::OnPaint( const FGeometry& AllottedGeometry, const FSlateRect& MyClippingRect, FSlateWindowElementList& OutDrawElements, int32 LayerId, const FWidgetStyle& InWidgetStyle, bool bParentEnabled ) const
+int32 SSequencerSectionAreaView::OnPaint( const FPaintArgs& Args, const FGeometry& AllottedGeometry, const FSlateRect& MyClippingRect, FSlateWindowElementList& OutDrawElements, int32 LayerId, const FWidgetStyle& InWidgetStyle, bool bParentEnabled ) const
 {
 	FArrangedChildren ArrangedChildren(EVisibility::Visible);
 	ArrangeChildren(AllottedGeometry, ArrangedChildren);
@@ -129,7 +129,7 @@ int32 SSequencerSectionAreaView::OnPaint( const FGeometry& AllottedGeometry, con
 	{
 		FArrangedWidget& CurWidget = ArrangedChildren(ChildIndex);
 		FSlateRect ChildClipRect = MyClippingRect.IntersectionWith( CurWidget.Geometry.GetClippingRect() );
-		const int32 CurWidgetsMaxLayerId = CurWidget.Widget->Paint( CurWidget.Geometry, ChildClipRect, OutDrawElements, LayerId+1, InWidgetStyle, ShouldBeEnabled( bParentEnabled ) );
+		const int32 CurWidgetsMaxLayerId = CurWidget.Widget->Paint( Args.WithNewParent(this), CurWidget.Geometry, ChildClipRect, OutDrawElements, LayerId+1, InWidgetStyle, ShouldBeEnabled( bParentEnabled ) );
 	}
 
 	return LayerId+1;
diff --git a/Engine/Source/Editor/Sequencer/Private/SSequencerSectionOverlay.cpp b/Engine/Source/Editor/Sequencer/Private/SSequencerSectionOverlay.cpp
--- a/Engine/Source/Editor/Sequencer/Private/SSequencerSectionOverlay.cpp
+++ b/Engine/Source/Editor/Sequencer/Private/SSequencerSectionOverlay.cpp
@@ -4,11 +4,11 @@
 #include "SSequencerSectionOverlay.h"
 #include "TimeSliderController.h"
 
-int32 SSequencerSectionOverlay::OnPaint( const FGeometry& AllottedGeometry, const FSlateRect& MyClippingRect, FSlateWindowElementList& OutDrawElements, int32 LayerId, const FWidgetStyle& InWidgetStyle, bool bParentEnabled ) const
+int32 SSequencerSectionOverlay::OnPaint( const FPaintArgs& Args, const FGeometry& AllottedGeometry, const FSlateRect& MyClippingRect, FSlateWindowElementList& OutDrawElements, int32 LayerId, const FWidgetStyle& InWidgetStyle, bool bParentEnabled ) const
 {
 	TimeSliderController->OnPaintSectionView( AllottedGeometry, MyClippingRect, OutDrawElements, LayerId, ShouldBeEnabled( bParentEnabled ), bDisplayTickLines.Get(), bDisplayScrubPosition.Get() );
 
-	return SCompoundWidget::OnPaint( AllottedGeometry, MyClippingRect, OutDrawElements, LayerId, InWidgetStyle, bParentEnabled  );
+	return SCompoundWidget::OnPaint( Args, AllottedGeometry, MyClippingRect, OutDrawElements, LayerId, InWidgetStyle, bParentEnabled  );
 }
 
 FReply SSequencerSectionOverlay::OnMouseButtonDown( const FGeometry& MyGeometry, const FPointerEvent& MouseEvent )
diff --git a/Engine/Source/Editor/Sequencer/Private/SequencerDisplayNode.cpp b/Engine/Source/Editor/Sequencer/Private/SequencerDisplayNode.cpp
--- a/Engine/Source/Editor/Sequencer/Private/SequencerDisplayNode.cpp
+++ b/Engine/Source/Editor/Sequencer/Private/SequencerDisplayNode.cpp
@@ -23,7 +23,7 @@ public:
 	SLATE_END_ARGS()
 
 	/** SLeafWidget Interface */
-	virtual int32 OnPaint( const FGeometry& AllottedGeometry, const FSlateRect& MyClippingRect, FSlateWindowElementList& OutDrawElements, int32 LayerId, const FWidgetStyle& InWidgetStyle, bool bParentEnabled ) const override;
+	virtual int32 OnPaint( const FPaintArgs& Args, const FGeometry& AllottedGeometry, const FSlateRect& MyClippingRect, FSlateWindowElementList& OutDrawElements, int32 LayerId, const FWidgetStyle& InWidgetStyle, bool bParentEnabled ) const override;
 	virtual FVector2D ComputeDesiredSize() const override;
 	
 	void Construct( const FArguments& InArgs, TSharedRef<FSequencerDisplayNode> InRootNode )
@@ -49,7 +49,7 @@ private:
 	TAttribute< TRange<float> > ViewRange;
 };
 
-int32 SSequencerObjectTrack::OnPaint(const FGeometry& AllottedGeometry, const FSlateRect& MyClippingRect, FSlateWindowElementList& OutDrawElements, int32 LayerId, const FWidgetStyle& InWidgetStyle, bool bParentEnabled) const
+int32 SSequencerObjectTrack::OnPaint(const FPaintArgs& Args, const FGeometry& AllottedGeometry, const FSlateRect& MyClippingRect, FSlateWindowElementList& OutDrawElements, int32 LayerId, const FWidgetStyle& InWidgetStyle, bool bParentEnabled) const
 {
 	// Draw a region around the entire section area
 	FSlateDrawElement::MakeBox( 
diff --git a/Engine/Source/Editor/SequencerWidgets/Private/STimeSlider.cpp b/Engine/Source/Editor/SequencerWidgets/Private/STimeSlider.cpp
--- a/Engine/Source/Editor/SequencerWidgets/Private/STimeSlider.cpp
+++ b/Engine/Source/Editor/SequencerWidgets/Private/STimeSlider.cpp
@@ -12,11 +12,11 @@ void STimeSlider::Construct( const STimeSlider::FArguments& InArgs, TSharedRef<I
 	bMirrorLabels = InArgs._MirrorLabels;
 }
 
-int32 STimeSlider::OnPaint( const FGeometry& AllottedGeometry, const FSlateRect& MyClippingRect, FSlateWindowElementList& OutDrawElements, int32 LayerId, const FWidgetStyle& InWidgetStyle, bool bParentEnabled ) const
+int32 STimeSlider::OnPaint( const FPaintArgs& Args, const FGeometry& AllottedGeometry, const FSlateRect& MyClippingRect, FSlateWindowElementList& OutDrawElements, int32 LayerId, const FWidgetStyle& InWidgetStyle, bool bParentEnabled ) const
 {
 	int32 NewLayer = TimeSliderController->OnPaintTimeSlider( bMirrorLabels, AllottedGeometry, MyClippingRect, OutDrawElements, LayerId, InWidgetStyle, bParentEnabled );
 
-	return FMath::Max( NewLayer, SCompoundWidget::OnPaint(AllottedGeometry, MyClippingRect, OutDrawElements, NewLayer, InWidgetStyle, ShouldBeEnabled( bParentEnabled ) ) );
+	return FMath::Max( NewLayer, SCompoundWidget::OnPaint(Args, AllottedGeometry, MyClippingRect, OutDrawElements, NewLayer, InWidgetStyle, ShouldBeEnabled( bParentEnabled ) ) );
 }
 
 FReply STimeSlider::OnMouseButtonDown( const FGeometry& MyGeometry, const FPointerEvent& MouseEvent )
diff --git a/Engine/Source/Editor/UMGEditor/Private/Designer/SDesignSurface.cpp b/Engine/Source/Editor/UMGEditor/Private/Designer/SDesignSurface.cpp
--- a/Engine/Source/Editor/UMGEditor/Private/Designer/SDesignSurface.cpp
+++ b/Engine/Source/Editor/UMGEditor/Private/Designer/SDesignSurface.cpp
@@ -129,12 +129,12 @@ void SDesignSurface::Construct(const FArguments& InArgs)
 	];
 }
 
-int32 SDesignSurface::OnPaint(const FGeometry& AllottedGeometry, const FSlateRect& MyClippingRect, FSlateWindowElementList& OutDrawElements, int32 LayerId, const FWidgetStyle& InWidgetStyle, bool bParentEnabled) const
+int32 SDesignSurface::OnPaint(const FPaintArgs& Args, const FGeometry& AllottedGeometry, const FSlateRect& MyClippingRect, FSlateWindowElementList& OutDrawElements, int32 LayerId, const FWidgetStyle& InWidgetStyle, bool bParentEnabled) const
 {
 	const FSlateBrush* BackgroundImage = FEditorStyle::GetBrush(TEXT("Graph.Panel.SolidBackground"));
 	PaintBackgroundAsLines(BackgroundImage, AllottedGeometry, MyClippingRect, OutDrawElements, LayerId);
 
-	SCompoundWidget::OnPaint(AllottedGeometry, MyClippingRect, OutDrawElements, LayerId, InWidgetStyle, bParentEnabled);
+	SCompoundWidget::OnPaint(Args, AllottedGeometry, MyClippingRect, OutDrawElements, LayerId, InWidgetStyle, bParentEnabled);
 
 	return LayerId;
 }
diff --git a/Engine/Source/Editor/UMGEditor/Private/Designer/SDesignerView.cpp b/Engine/Source/Editor/UMGEditor/Private/Designer/SDesignerView.cpp
--- a/Engine/Source/Editor/UMGEditor/Private/Designer/SDesignerView.cpp
+++ b/Engine/Source/Editor/UMGEditor/Private/Designer/SDesignerView.cpp
@@ -617,9 +617,9 @@ FVector2D SDesignerView::GetSelectionDesignerWidgetsLocation() const
 //	}
 //}
 
-int32 SDesignerView::OnPaint(const FGeometry& AllottedGeometry, const FSlateRect& MyClippingRect, FSlateWindowElementList& OutDrawElements, int32 LayerId, const FWidgetStyle& InWidgetStyle, bool bParentEnabled) const
+int32 SDesignerView::OnPaint(const FPaintArgs& Args, const FGeometry& AllottedGeometry, const FSlateRect& MyClippingRect, FSlateWindowElementList& OutDrawElements, int32 LayerId, const FWidgetStyle& InWidgetStyle, bool bParentEnabled) const
 {
-	SDesignSurface::OnPaint(AllottedGeometry, MyClippingRect, OutDrawElements, LayerId, InWidgetStyle, bParentEnabled);
+	SDesignSurface::OnPaint(Args, AllottedGeometry, MyClippingRect, OutDrawElements, LayerId, InWidgetStyle, bParentEnabled);
 
 	LayerId += 1000;
 
diff --git a/Engine/Source/Editor/UnrealEd/Private/SColorGradientEditor.cpp b/Engine/Source/Editor/UnrealEd/Private/SColorGradientEditor.cpp
--- a/Engine/Source/Editor/UnrealEd/Private/SColorGradientEditor.cpp
+++ b/Engine/Source/Editor/UnrealEd/Private/SColorGradientEditor.cpp
@@ -120,7 +120,7 @@ void SColorGradientEditor::Construct( const FArguments& InArgs )
 	ContextMenuPosition = FVector2D::ZeroVector;
 }
 
-int32 SColorGradientEditor::OnPaint( const FGeometry& AllottedGeometry, const FSlateRect& MyClippingRect, FSlateWindowElementList& OutDrawElements, int32 LayerId, const FWidgetStyle& InWidgetStyle, bool bParentEnabled ) const
+int32 SColorGradientEditor::OnPaint( const FPaintArgs& Args, const FGeometry& AllottedGeometry, const FSlateRect& MyClippingRect, FSlateWindowElementList& OutDrawElements, int32 LayerId, const FWidgetStyle& InWidgetStyle, bool bParentEnabled ) const
 {
 	const TSharedRef< FSlateFontMeasure > FontMeasureService = FSlateApplication::Get().GetRenderer()->GetFontMeasureService();
 
diff --git a/Engine/Source/Editor/UnrealEd/Private/SCurveEditor.cpp b/Engine/Source/Editor/UnrealEd/Private/SCurveEditor.cpp
--- a/Engine/Source/Editor/UnrealEd/Private/SCurveEditor.cpp
+++ b/Engine/Source/Editor/UnrealEd/Private/SCurveEditor.cpp
@@ -371,7 +371,7 @@ bool SCurveEditor::GetInputEditEnabled() const
 	return (SelectedKeys.Num() == 1);
 }
 
-int32 SCurveEditor::OnPaint(const FGeometry& AllottedGeometry, const FSlateRect& MyClippingRect, FSlateWindowElementList& OutDrawElements, int32 LayerId, const FWidgetStyle& InWidgetStyle, bool bParentEnabled ) const
+int32 SCurveEditor::OnPaint(const FPaintArgs& Args, const FGeometry& AllottedGeometry, const FSlateRect& MyClippingRect, FSlateWindowElementList& OutDrawElements, int32 LayerId, const FWidgetStyle& InWidgetStyle, bool bParentEnabled ) const
 {
 	// Rendering info
 	bool bEnabled = ShouldBeEnabled( bParentEnabled );
@@ -463,7 +463,7 @@ int32 SCurveEditor::OnPaint(const FGeometry& AllottedGeometry, const FSlateRect&
 	}
 
 	// Paint children
-	SCompoundWidget::OnPaint(CurveAreaGeometry, MyClippingRect, OutDrawElements, LayerId, InWidgetStyle, bParentEnabled);
+	SCompoundWidget::OnPaint(Args, CurveAreaGeometry, MyClippingRect, OutDrawElements, LayerId, InWidgetStyle, bParentEnabled);
 
 	PaintGridLines(CurveAreaGeometry, ScaleInfo, OutDrawElements, LayerId++, MyClippingRect, DrawEffects);
 
diff --git a/Engine/Source/Editor/UnrealEd/Public/SCurveEditor.h b/Engine/Source/Editor/UnrealEd/Public/SCurveEditor.h
--- a/Engine/Source/Editor/UnrealEd/Public/SCurveEditor.h
+++ b/Engine/Source/Editor/UnrealEd/Public/SCurveEditor.h
@@ -390,7 +390,7 @@ protected:
 	UNREALED_API virtual float GetTimeStep(FTrackScaleInfo &ScaleInfo) const;
 	
 	// SWidget interface
-	UNREALED_API virtual int32 OnPaint( const FGeometry& AllottedGeometry, const FSlateRect& MyClippingRect, 
+	UNREALED_API virtual int32 OnPaint( const FPaintArgs& Args, const FGeometry& AllottedGeometry, const FSlateRect& MyClippingRect, 
 		FSlateWindowElementList& OutDrawElements, int32 LayerId, const FWidgetStyle& InWidgetStyle, bool bParentEnabled ) const ;
 
 private:
diff --git a/Engine/Source/Editor/WorldBrowser/Private/Tiles/SWorldComposition.cpp b/Engine/Source/Editor/WorldBrowser/Private/Tiles/SWorldComposition.cpp
--- a/Engine/Source/Editor/WorldBrowser/Private/Tiles/SWorldComposition.cpp
+++ b/Engine/Source/Editor/WorldBrowser/Private/Tiles/SWorldComposition.cpp
@@ -274,9 +274,7 @@ public:
 	}
 	
 	/**  SWidget interface */
-	virtual int32 OnPaint(const FGeometry& AllottedGeometry, const FSlateRect& MyClippingRect, 
-							FSlateWindowElementList& OutDrawElements, int32 LayerId, 
-							const FWidgetStyle& InWidgetStyle, bool bParentEnabled ) const override
+	virtual int32 OnPaint(const FPaintArgs& Args, const FGeometry& AllottedGeometry, const FSlateRect& MyClippingRect, FSlateWindowElementList& OutDrawElements, int32 LayerId, const FWidgetStyle& InWidgetStyle, bool bParentEnabled ) const override
 	{
 		// First paint the background
 		{
@@ -306,7 +304,7 @@ public:
 			TSharedRef<SWorldTileItem> ChildNode = StaticCastSharedRef<SWorldTileItem>(CurWidget.Widget);
 		
 			ChildNode->bAffectedByMarquee = SelectionToVisualize->Contains(ChildNode->GetObjectBeingDisplayed());
-			LayerId = CurWidget.Widget->Paint(CurWidget.Geometry, MyClippingRect, OutDrawElements, NodesLayerId, InWidgetStyle, ShouldBeEnabled(bParentEnabled));
+			LayerId = CurWidget.Widget->Paint(Args.WithNewParent(this), CurWidget.Geometry, MyClippingRect, OutDrawElements, NodesLayerId, InWidgetStyle, ShouldBeEnabled(bParentEnabled));
 			ChildNode->bAffectedByMarquee = false;
 		}
 		
diff --git a/Engine/Source/Editor/WorldBrowser/Private/Tiles/SWorldTileItem.cpp b/Engine/Source/Editor/WorldBrowser/Private/Tiles/SWorldTileItem.cpp
--- a/Engine/Source/Editor/WorldBrowser/Private/Tiles/SWorldTileItem.cpp
+++ b/Engine/Source/Editor/WorldBrowser/Private/Tiles/SWorldTileItem.cpp
@@ -358,9 +358,7 @@ FVector2D SWorldTileItem::ComputeDesiredSize() const
 	return TileModel->GetLevelSize2D();
 }
 
-int32 SWorldTileItem::OnPaint(const FGeometry& AllottedGeometry, const FSlateRect& ClippingRect, 
-								FSlateWindowElementList& OutDrawElements, int32 LayerId, 
-								const FWidgetStyle& InWidgetStyle, bool bParentEnabled) const
+int32 SWorldTileItem::OnPaint(const FPaintArgs& Args, const FGeometry& AllottedGeometry, const FSlateRect& ClippingRect, FSlateWindowElementList& OutDrawElements, int32 LayerId, const FWidgetStyle& InWidgetStyle, bool bParentEnabled) const
 {
 	const bool bIsVisible = FSlateRect::DoRectanglesIntersect(AllottedGeometry.GetClippingRect(), ClippingRect);
 	
@@ -373,7 +371,7 @@ int32 SWorldTileItem::OnPaint(const FGeometry& AllottedGeometry, const FSlateRec
 			bNeedRefresh = false;
 		}
 		
-		LayerId = SNodePanel::SNode::OnPaint(AllottedGeometry, ClippingRect, OutDrawElements, LayerId, InWidgetStyle, bParentEnabled);
+		LayerId = SNodePanel::SNode::OnPaint(Args, AllottedGeometry, ClippingRect, OutDrawElements, LayerId, InWidgetStyle, bParentEnabled);
 		
 		const bool bSelected = (IsItemSelected() || bAffectedByMarquee);
 
diff --git a/Engine/Source/Editor/WorldBrowser/Private/Tiles/SWorldTileItem.h b/Engine/Source/Editor/WorldBrowser/Private/Tiles/SWorldTileItem.h
--- a/Engine/Source/Editor/WorldBrowser/Private/Tiles/SWorldTileItem.h
+++ b/Engine/Source/Editor/WorldBrowser/Private/Tiles/SWorldTileItem.h
@@ -89,9 +89,7 @@ private:
 	virtual TSharedPtr<IToolTip> GetToolTip() override;
 	virtual FVector2D ComputeDesiredSize() const override;
 	virtual bool OnHitTest(const FGeometry& MyGeometry, FVector2D InAbsoluteCursorPosition) override;
-	virtual int32 OnPaint(const FGeometry& AllottedGeometry, const FSlateRect& MyClippingRect, 
-					FSlateWindowElementList& OutDrawElements, int32 LayerId, 
-					const FWidgetStyle& InWidgetStyle, bool bParentEnabled) const override;
+	virtual int32 OnPaint(const FPaintArgs& Args, const FGeometry& AllottedGeometry, const FSlateRect& MyClippingRect, FSlateWindowElementList& OutDrawElements, int32 LayerId, const FWidgetStyle& InWidgetStyle, bool bParentEnabled) const override;
 	virtual FReply OnMouseButtonDoubleClick(const FGeometry& InMyGeometry, const FPointerEvent& InMouseEvent) override;
 	// SWidget interface end
 
diff --git a/Engine/Source/Runtime/Core/Private/Windows/WindowsApplication.cpp b/Engine/Source/Runtime/Core/Private/Windows/WindowsApplication.cpp
--- a/Engine/Source/Runtime/Core/Private/Windows/WindowsApplication.cpp
+++ b/Engine/Source/Runtime/Core/Private/Windows/WindowsApplication.cpp
@@ -1025,6 +1025,15 @@ int32 FWindowsApplication::ProcessMessage( HWND hwnd, uint32 msg, WPARAM wParam,
 			}
 			break;
 
+		case WM_DISPLAYCHANGE:
+			{
+				// Slate needs to know when desktop size changes.
+				FDisplayMetrics DisplayMetrics;
+				GetDisplayMetrics( DisplayMetrics );
+				BroadcastDisplayMetricsChanged(DisplayMetrics);
+			}
+			break;
+
 		case WM_GETDLGCODE:
 			{
 				// Slate wants all keys and messages.
diff --git a/Engine/Source/Runtime/Core/Public/GenericPlatform/GenericApplication.h b/Engine/Source/Runtime/Core/Public/GenericPlatform/GenericApplication.h
--- a/Engine/Source/Runtime/Core/Public/GenericPlatform/GenericApplication.h
+++ b/Engine/Source/Runtime/Core/Public/GenericPlatform/GenericApplication.h
@@ -270,13 +270,12 @@ namespace EWindowTitleAlignment
 	};
 }
 
-
 /**
  * Generic platform application interface
  */
 class GenericApplication 
 {
-public:	
+public:
 
 	GenericApplication( const TSharedPtr< ICursor >& InCursor )
 		: Cursor( InCursor )
@@ -330,6 +329,11 @@ public:
 
 	virtual bool TryCalculatePopupWindowPosition( const FPlatformRect& InAnchor, const FVector2D& InSize, const EPopUpOrientation::Type Orientation, /*OUT*/ FVector2D* const CalculatedPopUpPosition ) const { return false; }
 
+	DECLARE_EVENT_OneParam(GenericApplication, FOnDisplayMetricsChanged, const FDisplayMetrics&);
+	
+	/** Notifies subscribers when any of the display metrics change: e.g. resolution changes or monitor sare re-arranged. */
+	FOnDisplayMetricsChanged& OnDisplayMetricsChanged(){ return OnDisplayMetricsChangedEvent; }
+	
 	virtual void GetDisplayMetrics( FDisplayMetrics& OutDisplayMetrics ) const { }
 
 	virtual void GetInitialDisplayMetrics( FDisplayMetrics& OutDisplayMetrics ) const { GetDisplayMetrics(OutDisplayMetrics); }
@@ -362,4 +366,10 @@ public:
 protected:
 
 	TSharedRef< class FGenericApplicationMessageHandler > MessageHandler;
+	
+	/** Trigger the OnDisplayMetricsChanged event with the argument 'InMetrics' */
+	void BroadcastDisplayMetricsChanged( const FDisplayMetrics& InMetrics ){ OnDisplayMetricsChangedEvent.Broadcast( InMetrics ); }
+	
+	// Notifies subscribers when any of the display metrics change: e.g. resolution changes or monitor sare re-arranged.
+	FOnDisplayMetricsChanged OnDisplayMetricsChangedEvent;
 };
\ No newline at end of file
