Added analytics reporting to the profiler

#ttp 340626 - FIXIF: EDITOR: ANALYTICS: Profiler Usage Event

#branch UE4

#added Added analytics for tracking how long the profiler has been used for

#note It may have been possible to just use HasMouseCapture in the Tick, however this returns false if the tab isn't 'active' yet - so wouldn't count the time when a use is hovering over the window, mulling over their options.

diff --git a/Engine/Source/Developer/Profiler/Private/Widgets/SProfilerWindow.cpp b/Engine/Source/Developer/Profiler/Private/Widgets/SProfilerWindow.cpp
--- a/Engine/Source/Developer/Profiler/Private/Widgets/SProfilerWindow.cpp
+++ b/Engine/Source/Developer/Profiler/Private/Widgets/SProfilerWindow.cpp
@@ -8,6 +8,11 @@
 #include "SProfilerSettings.h"
 #include "DesktopPlatformModule.h"
 
+#if WITH_EDITOR
+	#include "Runtime/Analytics/Analytics/Public/Interfaces/IAnalyticsProvider.h"
+	#include "Runtime/Engine/Public/EngineAnalytics.h"
+#endif // WITH_EDITOR
+
 #define LOCTEXT_NAMESPACE "SProfilerWindow"
 
 static FText GetTextForNotification( const EProfilerNotificationTypes::Type NotificatonType, const ELoadingProgressStates::Type ProgressState, const FString& Filename, const float ProgressPercent = 0.0f )
@@ -75,6 +80,8 @@ static FText GetTextForNotification( const EProfilerNotificationTypes::Type Noti
 }
 
 SProfilerWindow::SProfilerWindow()
+	: DurationActive(0.0f)
+	, bIsActive(false)
 {}
 
 SProfilerWindow::~SProfilerWindow()
@@ -84,6 +91,13 @@ SProfilerWindow::~SProfilerWindow()
 	{
 		FProfilerManager::Get()->OnViewModeChanged().RemoveAll( this );
 	}
+
+#if WITH_EDITOR
+	if (DurationActive > 0.0f && FEngineAnalytics::IsAvailable())
+	{
+		FEngineAnalytics::GetProvider().RecordEvent(TEXT("Editor.Usage.Profiler"), FAnalyticsEventAttribute(TEXT("Duration"), DurationActive));
+	}
+#endif // WITH_EDITOR
 }
 
 BEGIN_SLATE_FUNCTION_BUILD_OPTIMIZATION
@@ -476,6 +490,30 @@ void SProfilerWindow::SendingServiceSideCapture_Load( const FString Filename )
 	}
 }
 
+void SProfilerWindow::Tick(const FGeometry& AllottedGeometry, const double InCurrentTime, const float InDeltaTime)
+{
+	SCompoundWidget::Tick(AllottedGeometry, InCurrentTime, InDeltaTime);
+
+	if (bIsActive)
+	{
+		DurationActive += InDeltaTime;
+	}
+}
+
+void SProfilerWindow::OnMouseEnter(const FGeometry& MyGeometry, const FPointerEvent& MouseEvent)
+{
+	SCompoundWidget::OnMouseEnter(MyGeometry, MouseEvent);
+
+	bIsActive = true;
+}
+
+void SProfilerWindow::OnMouseLeave(const FPointerEvent& MouseEvent)
+{
+	SCompoundWidget::OnMouseLeave(MouseEvent);
+
+	bIsActive = false;
+}
+
 FReply SProfilerWindow::OnKeyDown( const FGeometry& MyGeometry, const FKeyboardEvent& InKeyboardEvent )
 {
 	return FProfilerManager::Get()->GetCommandList()->ProcessCommandBindings( InKeyboardEvent ) ? FReply::Handled() : FReply::Unhandled();
diff --git a/Engine/Source/Developer/Profiler/Private/Widgets/SProfilerWindow.h b/Engine/Source/Developer/Profiler/Private/Widgets/SProfilerWindow.h
--- a/Engine/Source/Developer/Profiler/Private/Widgets/SProfilerWindow.h
+++ b/Engine/Source/Developer/Profiler/Private/Widgets/SProfilerWindow.h
@@ -61,6 +61,30 @@ protected:
 	void ProfilerManager_OnViewModeChanged( EProfilerViewMode::Type NewViewMode );
 
 private:
+	/**
+	 * Ticks this widget.  Override in derived classes, but always call the parent implementation.
+	 *
+	 * @param  AllottedGeometry The space allotted for this widget
+	 * @param  InCurrentTime  Current absolute real time
+	 * @param  InDeltaTime  Real time passed since last tick
+	 */
+	virtual void Tick( const FGeometry& AllottedGeometry, const double InCurrentTime, const float InDeltaTime ) override;
+
+	/**
+	 * The system will use this event to notify a widget that the cursor has entered it. This event is NOT bubbled.
+	 *
+	 * @param MyGeometry The Geometry of the widget receiving the event
+	 * @param MouseEvent Information about the input event
+	 */
+	virtual void OnMouseEnter(const FGeometry& MyGeometry, const FPointerEvent& MouseEvent) override;
+
+	/**
+	 * The system will use this event to notify a widget that the cursor has left it. This event is NOT bubbled.
+	 *
+	 * @param MouseEvent Information about the input event
+	 */
+	virtual void OnMouseLeave(const FPointerEvent& MouseEvent) override;
+
 	/**
 	 * Called after a key is pressed when this widget has keyboard focus
 	 *
@@ -117,6 +141,12 @@ public:
 	/** Active event graphs, one event graph for each profiler instance, stored as FGuid -> SEventGraph. */
 	TMap< FGuid, TSharedRef<class SEventGraph> > ActiveEventGraphs;
 
-
+	/** Widget for the profiler mini view */
 	TSharedPtr<SProfilerMiniView> ProfilerMiniView;
+
+	/** The number of seconds the profiler has been active */
+	float	DurationActive;
+
+	/** Whether or not the profiler widget is being used */
+	bool	bIsActive;
 };
\ No newline at end of file
diff --git a/Engine/Source/Developer/Profiler/Profiler.Build.cs b/Engine/Source/Developer/Profiler/Profiler.Build.cs
--- a/Engine/Source/Developer/Profiler/Profiler.Build.cs
+++ b/Engine/Source/Developer/Profiler/Profiler.Build.cs
@@ -28,6 +28,15 @@ public class Profiler : ModuleRules
 			}
 		);
 
+        if (UEBuildConfiguration.bBuildEditor)
+        {
+            PrivateDependencyModuleNames.AddRange(
+                new string[] {
+					"Engine",
+				}
+            );
+        }
+
 		PrivateDependencyModuleNames.AddRange(
 			new string[] {
 				"SlateCore",
