Profiler - Fixed FPSChart

diff --git a/Engine/Source/Developer/Profiler/Private/ProfilerSession.cpp b/Engine/Source/Developer/Profiler/Private/ProfilerSession.cpp
--- a/Engine/Source/Developer/Profiler/Private/ProfilerSession.cpp
+++ b/Engine/Source/Developer/Profiler/Private/ProfilerSession.cpp
@@ -324,6 +324,8 @@ bool FProfilerSession::HandleTicker( float DeltaTime )
 		FProfilerSampleArray& MutableCollection = const_cast<FProfilerSampleArray&>(DataProvider->GetCollection());
 		MutableCollection[FrameRootSampleIndex].SetDurationMS( GameThreadTimeMS != 0.0f ? GameThreadTimeMS : MaxThreadTimeMS );
 
+		FPSAnalyzer->AddSample( GameThreadTimeMS > 0.0f ? 1000.0f/GameThreadTimeMS : 0.0f );
+
 		// Process the non-hierarchical samples for the specified frame.
 		{
 			// Process integer counters.
@@ -339,7 +341,6 @@ bool FProfilerSession::HandleTicker( float DeltaTime )
 			{
 				const FProfilerFloatAccumulator& FloatCounter = CurrentProfilerData.FloatAccumulators[Index];
 				DataProvider->AddCounterSample( MetaData->GetStatByID(FloatCounter.StatId).OwningGroup().ID(), FloatCounter.StatId, (double)FloatCounter.Value, EProfilerSampleTypes::NumberFloat );
-
 			}
 		}
 
