Profiler - Removed unused code

diff --git a/Engine/Source/Developer/Profiler/Private/Widgets/SProfilerFPSChartPanel.cpp b/Engine/Source/Developer/Profiler/Private/Widgets/SProfilerFPSChartPanel.cpp
--- a/Engine/Source/Developer/Profiler/Private/Widgets/SProfilerFPSChartPanel.cpp
+++ b/Engine/Source/Developer/Profiler/Private/Widgets/SProfilerFPSChartPanel.cpp
@@ -7,14 +7,6 @@
 class SProfilerFPSStatisticsPanel : public SCompoundWidget
 {
 public:
-	/** Default constructor. */
-	SProfilerFPSStatisticsPanel()
-	{}
-
-	/** Virtual destructor */
-	virtual ~SProfilerFPSStatisticsPanel()
-	{}
-
 	SLATE_BEGIN_ARGS( SProfilerFPSStatisticsPanel )
 	{}
 		SLATE_ARGUMENT( TSharedPtr<FFPSAnalyzer>, FPSAnalyzer )
@@ -137,15 +129,6 @@ protected:
 	TSharedPtr<FFPSAnalyzer> FPSAnalyzer;
 };
 
-SProfilerFPSChartPanel::SProfilerFPSChartPanel()
-{
-	RestoreDefaultState();
-}
-
-SProfilerFPSChartPanel::~SProfilerFPSChartPanel()
-{
-}
-
 BEGIN_SLATE_FUNCTION_BUILD_OPTIMIZATION
 void SProfilerFPSChartPanel::Construct( const FArguments& InArgs )
 {
@@ -175,12 +158,4 @@ void SProfilerFPSChartPanel::Construct( const FArguments& InArgs )
 }
 END_SLATE_FUNCTION_BUILD_OPTIMIZATION
 
-void SProfilerFPSChartPanel::Tick( const FGeometry& AllottedGeometry, const double InCurrentTime, const float InDeltaTime )
-{
-}
-
-void SProfilerFPSChartPanel::RestoreDefaultState()
-{
-}
-
 #undef LOCTEXT_NAMESPACE
diff --git a/Engine/Source/Developer/Profiler/Private/Widgets/SProfilerFPSChartPanel.h b/Engine/Source/Developer/Profiler/Private/Widgets/SProfilerFPSChartPanel.h
--- a/Engine/Source/Developer/Profiler/Private/Widgets/SProfilerFPSChartPanel.h
+++ b/Engine/Source/Developer/Profiler/Private/Widgets/SProfilerFPSChartPanel.h
@@ -6,12 +6,6 @@
 class SProfilerFPSChartPanel : public SCompoundWidget
 {
 public:
-	/** Default constructor. */
-	SProfilerFPSChartPanel();
-		
-	/** Virtual destructor. */
-	virtual ~SProfilerFPSChartPanel();
-
 	SLATE_BEGIN_ARGS( SProfilerFPSChartPanel )
 		{}
 
@@ -24,11 +18,4 @@ public:
 	 * @param InArgs - The declaration data for this widget
 	 */
 	void Construct( const FArguments& InArgs );
-
-	// SWidget interface
-	virtual void Tick( const FGeometry& AllottedGeometry, const double InCurrentTime, const float InDeltaTime ) override;
-	// End of SWidget interface
-
-	/** Restores default state for the graph panel by removing all connected data sources and reseting view to default values. */
-	void RestoreDefaultState();
 };
\ No newline at end of file
