Profiler - Removed unused code

diff --git a/Engine/Source/Developer/Profiler/Private/Widgets/SHistogram.cpp b/Engine/Source/Developer/Profiler/Private/Widgets/SHistogram.cpp
--- a/Engine/Source/Developer/Profiler/Private/Widgets/SHistogram.cpp
+++ b/Engine/Source/Developer/Profiler/Private/Widgets/SHistogram.cpp
@@ -4,20 +4,10 @@
 
 #define LOCTEXT_NAMESPACE "SHistogram"
 
-
 /*-----------------------------------------------------------------------------
 	SHistogram
 -----------------------------------------------------------------------------*/
 
-SHistogram::SHistogram()
-{
-	RestoreDefaultState();
-}
-
-SHistogram::~SHistogram()
-{
-}
-
 void SHistogram::Construct( const FArguments& InArgs )
 {
 	Description = InArgs._Description.Get();
@@ -29,11 +19,6 @@ void SHistogram::Construct( const FArguments& InArgs )
 	];
 }
 
-void SHistogram::Tick( const FGeometry& AllottedGeometry, const double InCurrentTime, const float InDeltaTime )
-{
-	SCompoundWidget::Tick(AllottedGeometry,InCurrentTime,InDeltaTime);
-}
-
 int32 SHistogram::OnPaint( const FPaintArgs& Args, const FGeometry& AllottedGeometry, const FSlateRect& MyClippingRect, FSlateWindowElementList& OutDrawElements, int32 LayerId, const FWidgetStyle& InWidgetStyle, bool bParentEnabled ) const
 {
 	static double TotalTime = 0.0f;
@@ -183,49 +168,6 @@ int32 SHistogram::OnPaint( const FPaintArgs& Args, const FGeometry& AllottedGeom
 	return SCompoundWidget::OnPaint( Args, AllottedGeometry, MyClippingRect, OutDrawElements, LayerId, InWidgetStyle, bParentEnabled && IsEnabled() );
 }
 
-void SHistogram::RestoreDefaultState()
-{
-}
-
-FReply SHistogram::OnMouseButtonDown( const FGeometry& MyGeometry, const FPointerEvent& MouseEvent )
-{
-	return FReply::Handled();
-}
-
-FReply SHistogram::OnMouseButtonUp( const FGeometry& MyGeometry, const FPointerEvent& MouseEvent )
-{
-	return FReply::Unhandled();
-}
-
-FReply SHistogram::OnMouseMove( const FGeometry& MyGeometry, const FPointerEvent& MouseEvent )
-{
-	return FReply::Unhandled();
-}
-
-void SHistogram::OnMouseEnter( const FGeometry& MyGeometry, const FPointerEvent& MouseEvent )
-{
-
-}
-
-void SHistogram::OnMouseLeave( const FPointerEvent& MouseEvent )
-{
-}
-
-FReply SHistogram::OnMouseWheel( const FGeometry& MyGeometry, const FPointerEvent& MouseEvent )
-{
-	return FReply::Handled();
-}
-
-FReply SHistogram::OnMouseButtonDoubleClick( const FGeometry& MyGeometry, const FPointerEvent& MouseEvent )
-{
-	return FReply::Unhandled();
-}
-
-FCursorReply SHistogram::OnCursorQuery( const FGeometry& MyGeometry, const FPointerEvent& CursorEvent ) const 
-{
-	return FCursorReply::Unhandled();
-}
-
 #undef LOCTEXT_NAMESPACE
 
 
diff --git a/Engine/Source/Developer/Profiler/Private/Widgets/SHistogram.h b/Engine/Source/Developer/Profiler/Private/Widgets/SHistogram.h
--- a/Engine/Source/Developer/Profiler/Private/Widgets/SHistogram.h
+++ b/Engine/Source/Developer/Profiler/Private/Widgets/SHistogram.h
@@ -12,12 +12,6 @@ typedef TSharedPtr<class IHistogramDataSource> FHistogramDataSourcePtr;
 /** Type definition for shared references to instances of FGraphDataSource. */
 typedef TSharedRef<class IHistogramDataSource> FHistogramDataSourceRef;
 
-
-/*-----------------------------------------------------------------------------
-	Delegates
------------------------------------------------------------------------------*/
-
-
 /*-----------------------------------------------------------------------------
 	Basic structures
 -----------------------------------------------------------------------------*/
@@ -93,12 +87,6 @@ public:
 class SHistogram : public SCompoundWidget
 {
 public:
-	/** Default constructor. */
-	SHistogram();
-
-	/** Virtual destructor. */
-	virtual ~SHistogram();
-
 	SLATE_BEGIN_ARGS( SHistogram )
 		: _Description()
 		{}
@@ -115,29 +103,8 @@ public:
 	 */
 	void Construct( const FArguments& InArgs );
 
-	virtual void Tick( const FGeometry& AllottedGeometry, const double InCurrentTime, const float InDeltaTime ) override;
-
 	virtual int32 OnPaint( const FPaintArgs& Args, const FGeometry& AllottedGeometry, const FSlateRect& MyClippingRect, FSlateWindowElementList& OutDrawElements, int32 LayerId, const FWidgetStyle& InWidgetStyle, bool bParentEnabled ) const override;
 
-	virtual FReply OnMouseButtonDown( const FGeometry& MyGeometry, const FPointerEvent& MouseEvent ) override;
-	
-	virtual FReply OnMouseButtonUp( const FGeometry& MyGeometry, const FPointerEvent& MouseEvent ) override;
-	
-	virtual FReply OnMouseMove( const FGeometry& MyGeometry, const FPointerEvent& MouseEvent ) override;
-	
-	virtual void OnMouseEnter( const FGeometry& MyGeometry, const FPointerEvent& MouseEvent ) override;
-	
-	virtual void OnMouseLeave( const FPointerEvent& MouseEvent ) override;
-	
-	virtual FReply OnMouseWheel( const FGeometry& MyGeometry, const FPointerEvent& MouseEvent ) override;
-
-	virtual FReply OnMouseButtonDoubleClick( const FGeometry& MyGeometry, const FPointerEvent& MouseEvent ) override;
-
-	virtual FCursorReply OnCursorQuery( const FGeometry& MyGeometry, const FPointerEvent& CursorEvent ) const override;
-
-	/** Restores default state for the data graph by removing all inner graphs. */
-	void RestoreDefaultState();
-
 	virtual FVector2D ComputeDesiredSize() const override
 	{
 		//return SCompoundWidget::ComputeDesiredSize();
