Fixed code relying on SLATE_TEXT_ATTRIBUTE for SHeaderRow::FColumn

Made sure everything was using FText rather than FString.

diff --git a/Engine/Source/Developer/ModuleUI/Private/SModuleUI.cpp b/Engine/Source/Developer/ModuleUI/Private/SModuleUI.cpp
--- a/Engine/Source/Developer/ModuleUI/Private/SModuleUI.cpp
+++ b/Engine/Source/Developer/ModuleUI/Private/SModuleUI.cpp
@@ -25,10 +25,10 @@ void SModuleUI::Construct(const SModuleUI::FArguments& InArgs)
 					(
 						SNew(SHeaderRow)
 						+SHeaderRow::Column("ModuleName")
-						.DefaultLabel(NSLOCTEXT("ModuleUI", "ModuleName", "Module").ToString())
+						.DefaultLabel(NSLOCTEXT("ModuleUI", "ModuleName", "Module"))
 						.FillWidth(200)
 						+SHeaderRow::Column("ModuleActions")
-						.DefaultLabel(NSLOCTEXT("ModuleUI", "ModuleActions", "Actions").ToString())
+						.DefaultLabel(NSLOCTEXT("ModuleUI", "ModuleActions", "Actions"))
 						.FillWidth(1000)
 					)
 			]
@@ -77,7 +77,7 @@ TSharedRef<ITableRow> SModuleUI::OnGenerateWidgetForModuleListView(TSharedPtr< F
 			{
 				return
 					SNew( STextBlock )
-					.Text( Item->ModuleName.ToString() );
+					.Text( FText::FromName(Item->ModuleName) );
 			}
 			else if ( ColumnName == "ModuleActions" )
 			{
diff --git a/Engine/Source/Developer/Profiler/Private/Widgets/SEventGraph.cpp b/Engine/Source/Developer/Profiler/Private/Widgets/SEventGraph.cpp
--- a/Engine/Source/Developer/Profiler/Private/Widgets/SEventGraph.cpp
+++ b/Engine/Source/Developer/Profiler/Private/Widgets/SEventGraph.cpp
@@ -65,8 +65,8 @@ namespace EventGraphColumns
 		(
 			EEventPropertyIndex::StatName,
 			TEXT("name"),
-			TEXT("Event Name"),
-			TEXT("Name of the event"),
+			LOCTEXT("EventNameColumnTitle", "Event Name"),
+			LOCTEXT("EventNameColumnDesc", "Name of the event"),
 			false,true,true,false,false,
 			HAlign_Left,
 			0.0f
@@ -76,8 +76,8 @@ namespace EventGraphColumns
 		(
 			EEventPropertyIndex::InclusiveTimeMS,
 			TEXT("inc"),
-			TEXT("Inc Time (MS)"),
-			TEXT("Duration of the sample and its children, in milliseconds"),
+			LOCTEXT("InclusiveTimeMSTitle", "Inc Time (MS)"),
+			LOCTEXT("InclusiveTimeMSDesc", "Duration of the sample and its children, in milliseconds"),
 			false,true,true,true,true,
 			HAlign_Right,
 			0.0f
@@ -87,8 +87,8 @@ namespace EventGraphColumns
 		(
 			EEventPropertyIndex::InclusiveTimePct,
 			TEXT("inc%"),
-			TEXT("Inc Time (%)"),
-			TEXT("Duration of the sample and its children as percent of the caller"),
+			LOCTEXT("InclusiveTimePercentageTitle", "Inc Time (%)"),
+			LOCTEXT("InclusiveTimePercentageDesc", "Duration of the sample and its children as percent of the caller"),
 			true,true,true,false,false,
 			HAlign_Right,
 			64.0f
@@ -98,8 +98,8 @@ namespace EventGraphColumns
 		(
 			EEventPropertyIndex::MinInclusiveTimeMS,
 			TEXT("min"),
-			TEXT("Min Inc Time (MS)"),
-			TEXT("Minimum inclusive time of all instances for this event, in milliseconds"),
+			LOCTEXT("MinInclusiveTimeMSTitle", "Min Inc Time (MS)"),
+			LOCTEXT("MinInclusiveTimeMSDesc", "Minimum inclusive time of all instances for this event, in milliseconds"),
 			true,false,true,true,false,
 			HAlign_Right,
 			0.0f
@@ -109,8 +109,8 @@ namespace EventGraphColumns
 		(
 			EEventPropertyIndex::MaxInclusiveTimeMS,
 			TEXT("max"),
-			TEXT("Max Inc Time (MS)"),
-			TEXT("Maximum inclusive time of all instances for this event, in milliseconds"),
+			LOCTEXT("MaxInclusiveTimeMSTitle", "Max Inc Time (MS)"),
+			LOCTEXT("MaxInclusiveTimeMSDesc", "Maximum inclusive time of all instances for this event, in milliseconds"),
 			true,false,true,true,false,
 			HAlign_Right,
 			0.0f
@@ -120,8 +120,8 @@ namespace EventGraphColumns
 		(
 			EEventPropertyIndex::AvgInclusiveTimeMS,
 			TEXT("avg"),
-			TEXT("Avg Inc Time (MS)"),
-			TEXT("Average inclusive time of all instances for this event, in milliseconds"),
+			LOCTEXT("AverageInclusiveTimeMSTitle", "Avg Inc Time (MS)"),
+			LOCTEXT("AverageInclusiveTimeMSDesc", "Average inclusive time of all instances for this event, in milliseconds"),
 			true,false,true,true,false,
 			HAlign_Right,
 			0.0f
@@ -130,8 +130,9 @@ namespace EventGraphColumns
 		FEventGraphColumn
 		(
 			EEventPropertyIndex::ExclusiveTimeMS,
-			TEXT("exc"),TEXT("Exc Time (MS)"),
-			TEXT("Exclusive time of this event, in milliseconds"),
+			TEXT("exc"),
+			LOCTEXT("ExclusiveTimeMSTitle", "Exc Time (MS)"),
+			LOCTEXT("ExclusiveTimeMSDesc", "Exclusive time of this event, in milliseconds"),
 			false,true,true,true,false,
 			HAlign_Right,
 			0.0f
@@ -140,8 +141,9 @@ namespace EventGraphColumns
 		FEventGraphColumn
 		(
 			EEventPropertyIndex::ExclusiveTimePct,
-			TEXT("exc%"),TEXT("Exc Time (%)"),
-			TEXT("Exclusive time of this event as percent of this call's inclusive time"),
+			TEXT("exc%"),
+			LOCTEXT("ExclusiveTimePercentageTitle", "Exc Time (%)"),
+			LOCTEXT("ExclusiveTimePercentageDesc", "Exclusive time of this event as percent of this call's inclusive time"),
 			true,true,true,false,false,
 			HAlign_Right,
 			64.0f
@@ -151,8 +153,8 @@ namespace EventGraphColumns
 		(
 			EEventPropertyIndex::AvgInclusiveTimePerCallMS,
 			TEXT("avgper"),
-			TEXT("Avg Inc Time per Call (MS)"),
-			TEXT("Average inclusive time per call of all instances for this event, in milliseconds"),
+			LOCTEXT("AverageInclusiveTimePerCallMSTitle", "Avg Inc Time per Call (MS)"),
+			LOCTEXT("AverageInclusiveTimePerCallMSDesc", "Average inclusive time per call of all instances for this event, in milliseconds"),
 			true,false,true,false,false,
 			HAlign_Right,
 			0.0f
@@ -162,8 +164,8 @@ namespace EventGraphColumns
 		(
 			EEventPropertyIndex::NumCallsPerFrame,
 			TEXT("calls"),
-			TEXT("Calls"),
-			TEXT("Number of times this event was called"),
+			LOCTEXT("CallsPerFrameTitle", "Calls"),
+			LOCTEXT("CallsPerFrameDesc", "Number of times this event was called"),
 			false,true,true,true,false,
 			HAlign_Right,
 			64.0f
@@ -173,8 +175,8 @@ namespace EventGraphColumns
 		(
 			EEventPropertyIndex::AvgNumCallsPerFrame,
 			TEXT("avgcalls"),
-			TEXT("Avg Calls"),
-			TEXT("Average number of times this event was called"),
+			LOCTEXT("AverageCallsPerFrameTitle", "Avg Calls"),
+			LOCTEXT("AverageCallsPerFrameDesc", "Average number of times this event was called"),
 			true,true,true,true,false,
 			HAlign_Right,
 			64.0f
@@ -185,8 +187,8 @@ namespace EventGraphColumns
 		(
 			EEventPropertyIndex::ThreadName,
 			TEXT("th"),
-			TEXT("Thread Name"),
-			TEXT("Name of the thread that this event was captured on"),
+			LOCTEXT("ThreadNameTitle", "Thread Name"),
+			LOCTEXT("ThreadNameDesc", "Name of the thread that this event was captured on"),
 			true,false,true,false,false,
 			HAlign_Left,
 			0.0f
@@ -197,8 +199,8 @@ namespace EventGraphColumns
 		(
 			EEventPropertyIndex::ThreadDurationMS,
 			TEXT("td"),
-			TEXT("Thread Duration (MS)"),
-			TEXT("Duration of the thread that this event was captured on, in milliseconds"),
+			LOCTEXT("ThreadDurationTitle", "Thread Duration (MS)"),
+			LOCTEXT("ThreadDurationDesc", "Duration of the thread that this event was captured on, in milliseconds"),
 			true,false,true,false,false,
 			HAlign_Right,
 			0.0f
@@ -209,8 +211,8 @@ namespace EventGraphColumns
 		(
 			EEventPropertyIndex::FrameDurationMS,
 			TEXT("fd"),
-			TEXT("Frame Duration (MS)"),
-			TEXT("Duration of the frame that this event was captured on, in milliseconds"),
+			LOCTEXT("FrameDurationMSTitle", "Frame Duration (MS)"),
+			LOCTEXT("FrameDurationMSDesc", "Duration of the frame that this event was captured on, in milliseconds"),
 			true,false,true,false,false,
 			HAlign_Right,
 			0.0f
@@ -221,8 +223,8 @@ namespace EventGraphColumns
 		(
 			EEventPropertyIndex::ThreadPct,
 			TEXT("t%"),
-			TEXT("% of Thread"),
-			TEXT("Percent of inclusive time spent by this event in the particular thread"),
+			LOCTEXT("ThreadPercentageTitle", "% of Thread"),
+			LOCTEXT("ThreadPercentageDesc", "Percent of inclusive time spent by this event in the particular thread"),
 			true,true,true,false,false,
 			HAlign_Right,
 			64.0f
@@ -233,8 +235,8 @@ namespace EventGraphColumns
 		(
 			EEventPropertyIndex::FramePct,
 			TEXT("f%"),
-			TEXT("% of Frame"),
-			TEXT("Percent of inclusive time spent by this event in the particular frame"),
+			LOCTEXT("FramePercentageTitle", "% of Frame"),
+			LOCTEXT("FramePercentageDesc", "Percent of inclusive time spent by this event in the particular frame"),
 			true,true,true,false,false,
 			HAlign_Right,
 			64.0f
@@ -245,8 +247,8 @@ namespace EventGraphColumns
 		(
 			EEventPropertyIndex::ThreadToFramePct,
 			TEXT("t2f%"),
-			TEXT("Thread to Frame (%)"),
-			TEXT("Percent of time spent in the thread in relation to the entire frame"),
+			LOCTEXT("ThreadToFramePercentageTitle", "Thread to Frame (%)"),
+			LOCTEXT("ThreadToFramePercentageDesc", "Percent of time spent in the thread in relation to the entire frame"),
 			true,false,true,false,false,
 			HAlign_Right,
 			64.0f
@@ -256,8 +258,8 @@ namespace EventGraphColumns
 		(
 			EEventPropertyIndex::StartTimeMS,
 			TEXT("start"),
-			TEXT("Start Time (MS)"),
-			TEXT("Start time of this event, in milliseconds"),
+			LOCTEXT("StartTimeMSTitle", "Start Time (MS)"),
+			LOCTEXT("StartTimeMSDesc", "Start time of this event, in milliseconds"),
 			true,false,true,false,false,
 			HAlign_Right,
 			0.0f
@@ -267,8 +269,8 @@ namespace EventGraphColumns
 		(
 			EEventPropertyIndex::GroupName,
 			TEXT("grp"),
-			TEXT("Group Name"),
-			TEXT("Name of the stat group this event belongs to"),
+			LOCTEXT("GroupNameTitle", "Group Name"),
+			LOCTEXT("GroupNameDesc", "Name of the stat group this event belongs to"),
 			true,false,true,false,false,
 			HAlign_Left,
 			0.0f
@@ -279,8 +281,8 @@ namespace EventGraphColumns
 		(
 			EEventPropertyIndex::None,
 			TEXT("None"),
-			TEXT("None"),
-			TEXT("None"),
+			LOCTEXT("None", "None"),
+			LOCTEXT("None", "None"),
 			false,false,false,false,false,
 			HAlign_Left,
 			0.0f
@@ -1540,7 +1542,7 @@ TSharedPtr<SWidget> SEventGraph::EventGraph_GetMenuContent() const
 	else if( NumSelectedEvents == 1 )
 	{
 		SelectionStr = FText::FromString( SelectedEvent->_StatName.GetPlainNameString() );
-		PropertyName = FText::FromString( Column.ShortName );
+		PropertyName = Column.ShortName;
 		PropertyValue = FText::FromString( SelectedEvent->GetFormattedValue(Column.Index) );
 	}
 	else
@@ -1858,8 +1860,8 @@ void SEventGraph::EventGraph_BuildSortByMenu( FMenuBuilder& MenuBuilder )
 			);
 			MenuBuilder.AddMenuEntry
 			( 
-				FText::FromString( Column.ShortName ),
-				FText::FromString( Column.Description ), 
+				Column.ShortName,
+				Column.Description, 
 				FSlateIcon(), Action_SortByColumn, NAME_None, EUserInterfaceActionType::RadioButton 
 			);
 		}
@@ -1915,8 +1917,8 @@ void SEventGraph::EventGraph_BuildViewColumnMenu( FMenuBuilder& MenuBuilder )
 		);
 		MenuBuilder.AddMenuEntry
 		( 
-			FText::FromString( Column.ShortName ),
-			FText::FromString( Column.Description ), 
+			Column.ShortName,
+			Column.Description, 
 			FSlateIcon(), Action_ToggleColumn, NAME_None, EUserInterfaceActionType::ToggleButton 
 		);
 	}
@@ -2292,7 +2294,7 @@ void SEventGraph::ContextMenu_CopySelectedToClipboard_Execute()
 	for( uint32 ColumnIndex = 0; ColumnIndex < EventGraphColumns::NumColumns; ColumnIndex++ )
 	{
 		const FEventGraphColumn& Column = EventGraphColumns::Collection[ColumnIndex];
-		Result += FString::Printf( TEXT("\"%s\","), *Column.ShortName );
+		Result += FString::Printf( TEXT("\"%s\","), *Column.ShortName.ToString() );
 	}
 	Result += LINE_TERMINATOR;
 
diff --git a/Engine/Source/Developer/Profiler/Private/Widgets/SEventGraph.h b/Engine/Source/Developer/Profiler/Private/Widgets/SEventGraph.h
--- a/Engine/Source/Developer/Profiler/Private/Widgets/SEventGraph.h
+++ b/Engine/Source/Developer/Profiler/Private/Widgets/SEventGraph.h
@@ -71,8 +71,8 @@ public:
 	(
 		const EEventPropertyIndex::Type InIndex,
 		const FName InSearchID,
-		const FString InShortName,
-		const FString InDescription,
+		const FText InShortName,
+		const FText InDescription,
 		const bool bInCanBeHidden,
 		const bool bInIsVisible,
 		const bool bInCanBeSorted,
@@ -84,8 +84,8 @@ public:
 		: Index( InIndex )
 		, ID( FEventGraphSample::GetEventPropertyByIndex(InIndex).Name )
 		, SearchID( InSearchID )
-		, ShortName( InShortName )
-		, Description( InDescription )
+		, ShortName( MoveTemp(InShortName) )
+		, Description( MoveTemp(InDescription) )
 		, bCanBeHidden( bInCanBeHidden )
 		, bIsVisible( bInIsVisible )
 		, bCanBeSorted( bInCanBeSorted )
@@ -106,10 +106,10 @@ public:
 	FName SearchID;
 
 	/** Short name of the column, displayed in the event graph header. */
-	FString ShortName;
+	FText ShortName;
 
 	/** Long name of the column, displayed in the column tooltip. */
-	FString Description;
+	FText Description;
 
 	/** Whether this column can be hidden. */
 	bool bCanBeHidden;
diff --git a/Engine/Source/Developer/ScreenShotComparison/Private/Widgets/SScreenViewRow.cpp b/Engine/Source/Developer/ScreenShotComparison/Private/Widgets/SScreenViewRow.cpp
--- a/Engine/Source/Developer/ScreenShotComparison/Private/Widgets/SScreenViewRow.cpp
+++ b/Engine/Source/Developer/ScreenShotComparison/Private/Widgets/SScreenViewRow.cpp
@@ -21,7 +21,7 @@ void SScreenViewRow::Construct( const FArguments& InArgs, const TSharedRef<STabl
 	TSharedPtr<IScreenShotData> ScreenShotData = InArgs._ScreenShotData;
 
 	// The display string for the view
-	FString InfoText = FString::Printf( *LOCTEXT("ScreenShotName", "Screen shot: %s").ToString(), *ScreenShotData->GetName() );
+	const FText InfoText = FText::Format( LOCTEXT("ScreenShotNameFmt", "Screen shot: {0}"), FText::FromString(ScreenShotData->GetName()) );
 
 	ChildSlot
 	[
@@ -57,16 +57,16 @@ void SScreenViewRow::Construct( const FArguments& InArgs, const TSharedRef<STabl
 						.HAlignCell(HAlign_Center)
 						.VAlignCell(VAlign_Center)
 						.FixedWidth( 200.0f )
-						.DefaultLabel(LOCTEXT("PlatformName", "Platform").ToString())
+						.DefaultLabel(LOCTEXT("PlatformName", "Platform"))
 					+ SHeaderRow::Column( "Current View" )
 						.HAlignHeader(HAlign_Center)
 						.VAlignHeader(VAlign_Center)
 						.HAlignCell(HAlign_Center)
 						.VAlignCell(VAlign_Center)
  						.FixedWidth( 300.0f )
-						.DefaultLabel(LOCTEXT("CurrentViewName", "Current View").ToString())
+						.DefaultLabel(LOCTEXT("CurrentViewName", "Current View"))
 					+ SHeaderRow::Column( "History View" )
-					.DefaultLabel(LOCTEXT("HistoryViewName", "History").ToString())
+					.DefaultLabel(LOCTEXT("HistoryViewName", "History"))
 				)
 			]
 		]
diff --git a/Engine/Source/Developer/SessionFrontend/Private/Widgets/Browser/SSessionBrowserInstanceList.h b/Engine/Source/Developer/SessionFrontend/Private/Widgets/Browser/SSessionBrowserInstanceList.h
--- a/Engine/Source/Developer/SessionFrontend/Private/Widgets/Browser/SSessionBrowserInstanceList.h
+++ b/Engine/Source/Developer/SessionFrontend/Private/Widgets/Browser/SSessionBrowserInstanceList.h
@@ -56,19 +56,19 @@ public:
 					SNew(SHeaderRow)
 
 					+ SHeaderRow::Column("Name")
-						.DefaultLabel(LOCTEXT("InstanceListNameColumnHeader", "Instance").ToString())
+						.DefaultLabel(LOCTEXT("InstanceListNameColumnHeader", "Instance"))
 
 					+ SHeaderRow::Column("Type")
-						.DefaultLabel(LOCTEXT("InstanceListTypeColumnHeader", "Type").ToString())
+						.DefaultLabel(LOCTEXT("InstanceListTypeColumnHeader", "Type"))
 
 					+ SHeaderRow::Column("Device")
-						.DefaultLabel(LOCTEXT("InstanceListDeviceColumnHeader", "Device").ToString())
+						.DefaultLabel(LOCTEXT("InstanceListDeviceColumnHeader", "Device"))
 
 					+ SHeaderRow::Column("Platform")
-						.DefaultLabel(LOCTEXT("InstanceListPlatformColumnHeader", "Platform").ToString())
+						.DefaultLabel(LOCTEXT("InstanceListPlatformColumnHeader", "Platform"))
 
 					+ SHeaderRow::Column("Status")
-						.DefaultLabel(LOCTEXT("InstanceListStatusColumnHeader", "Status").ToString())
+						.DefaultLabel(LOCTEXT("InstanceListStatusColumnHeader", "Status"))
 						.HAlignCell(HAlign_Right)
 						.HAlignHeader(HAlign_Right)
 				)
diff --git a/Engine/Source/Developer/SessionFrontend/Private/Widgets/Console/SSessionConsole.cpp b/Engine/Source/Developer/SessionFrontend/Private/Widgets/Console/SSessionConsole.cpp
--- a/Engine/Source/Developer/SessionFrontend/Private/Widgets/Console/SSessionConsole.cpp
+++ b/Engine/Source/Developer/SessionFrontend/Private/Widgets/Console/SSessionConsole.cpp
@@ -91,19 +91,19 @@ void SSessionConsole::Construct( const FArguments& InArgs, ISessionManagerRef In
 												SNew(SHeaderRow)
 
 												+ SHeaderRow::Column("Verbosity")
-													.DefaultLabel(LOCTEXT("LogListVerbosityColumnHeader", " ").ToString())
+													.DefaultLabel(LOCTEXT("LogListVerbosityColumnHeader", " "))
 													.FixedWidth(24.0f)
 
 												+ SHeaderRow::Column("Instance")
-													.DefaultLabel(LOCTEXT("LogListHostNameColumnHeader", "Instance").ToString())
+													.DefaultLabel(LOCTEXT("LogListHostNameColumnHeader", "Instance"))
 													.FillWidth(0.20f)
 
 												+ SHeaderRow::Column("TimeSeconds")
-													.DefaultLabel(LOCTEXT("LogListTimestampColumnHeader", "Seconds").ToString())
+													.DefaultLabel(LOCTEXT("LogListTimestampColumnHeader", "Seconds"))
 													.FillWidth(0.10f)
 
 												+ SHeaderRow::Column("Message")
-													.DefaultLabel(LOCTEXT("LogListTextColumnHeader", "Message").ToString())
+													.DefaultLabel(LOCTEXT("LogListTextColumnHeader", "Message"))
 													.FillWidth(0.70f)
 											)
 									]
diff --git a/Engine/Source/Developer/SlateReflector/Private/Widgets/SWidgetReflector.cpp b/Engine/Source/Developer/SlateReflector/Private/Widgets/SWidgetReflector.cpp
--- a/Engine/Source/Developer/SlateReflector/Private/Widgets/SWidgetReflector.cpp
+++ b/Engine/Source/Developer/SlateReflector/Private/Widgets/SWidgetReflector.cpp
@@ -206,7 +206,7 @@ void SWidgetReflector::Construct( const FArguments& InArgs )
 							SNew(SHeaderRow)
 
 							+ SHeaderRow::Column("WidgetName")
-							.DefaultLabel(LOCTEXT("WidgetName", "Widget Name").ToString())
+							.DefaultLabel(LOCTEXT("WidgetName", "Widget Name"))
 							.FillWidth(0.65f)
 
 							+ SHeaderRow::Column("ForegroundColor")
@@ -215,16 +215,16 @@ void SWidgetReflector::Construct( const FArguments& InArgs )
 							.HeaderContent()
 							[
 								SNew(STextBlock)
-								.Text(LOCTEXT("ForegroundColor", "FG").ToString())
-								.ToolTipText(LOCTEXT("ForegroundColorToolTip", "Foreground Color").ToString())
+								.Text(LOCTEXT("ForegroundColor", "FG"))
+								.ToolTipText(LOCTEXT("ForegroundColorToolTip", "Foreground Color"))
 							]
 
 							+ SHeaderRow::Column("Visibility")
-							.DefaultLabel(LOCTEXT("Visibility", "Visibility" ).ToString())
+							.DefaultLabel(LOCTEXT("Visibility", "Visibility" ))
 							.FixedWidth(125.0f)
 
 							+ SHeaderRow::Column("WidgetInfo")
-							.DefaultLabel(LOCTEXT("WidgetInfo", "Widget Info" ).ToString())
+							.DefaultLabel(LOCTEXT("WidgetInfo", "Widget Info" ))
 							.FillWidth(0.25f)
 
 							+ SHeaderRow::Column("Address")
diff --git a/Engine/Source/Developer/TaskGraph/Private/SEventsTree.cpp b/Engine/Source/Developer/TaskGraph/Private/SEventsTree.cpp
--- a/Engine/Source/Developer/TaskGraph/Private/SEventsTree.cpp
+++ b/Engine/Source/Developer/TaskGraph/Private/SEventsTree.cpp
@@ -167,7 +167,7 @@ void SEventsTree::Construct( const FArguments& InArgs )
 					]
 
 					+ SHeaderRow::Column( NAME_DurationColumn )
-					.DefaultLabel( TAttribute< FString >::Create( TAttribute< FString >::FGetter::CreateSP( this, &SEventsTree::GetDurationColumnTitle ) ) ) 
+					.DefaultLabel( this, &SEventsTree::GetDurationColumnTitle ) 
 					.SortMode( TAttribute< EColumnSortMode::Type >::Create( TAttribute< EColumnSortMode::Type >::FGetter::CreateSP( this, &SEventsTree::GetColumnSortMode, NAME_DurationColumn ) ) )
 					.OnSort( FOnSortModeChanged::CreateSP( this, &SEventsTree::OnColumnSortModeChanged ) )
 					.MenuContent()
@@ -191,11 +191,11 @@ EColumnSortMode::Type SEventsTree::GetColumnSortMode( const FName ColumnId )
 	return SortMode;
 }
 
-FString SEventsTree::GetDurationColumnTitle() const
+FText SEventsTree::GetDurationColumnTitle() const
 {
 	static const FText Units[] = { NSLOCTEXT("TaskGraph", "microseconds", "microseconds"), NSLOCTEXT("TaskGraph", "milliseconds", "ms"), NSLOCTEXT("TaskGraph", "seconds", "s") };
 
-	return FText::Format( NSLOCTEXT("TaskGraph", "ColumnDurationValue", "Duration ({0})"), Units[ DurationUnits ] ).ToString();
+	return FText::Format( NSLOCTEXT("TaskGraph", "ColumnDurationValue", "Duration ({0})"), Units[ DurationUnits ] );
 }
 
 void ClearEventsSelection( TArray< TSharedPtr< FVisualizerEvent > >& Events )
diff --git a/Engine/Source/Developer/TaskGraph/Public/SEventsTree.h b/Engine/Source/Developer/TaskGraph/Public/SEventsTree.h
--- a/Engine/Source/Developer/TaskGraph/Public/SEventsTree.h
+++ b/Engine/Source/Developer/TaskGraph/Public/SEventsTree.h
@@ -228,7 +228,7 @@ protected:
 	}
 
 	/** Gets the currently selected time units text */
-	FString GetDurationColumnTitle() const;
+	FText GetDurationColumnTitle() const;
 
 	/** A pointer to the ListView of profiler events */
 	TSharedPtr< STreeView< TSharedPtr< FVisualizerEvent > > > EventsListView;
diff --git a/Engine/Source/Editor/Kismet/Private/Debugging/SKismetDebuggingView.cpp b/Engine/Source/Editor/Kismet/Private/Debugging/SKismetDebuggingView.cpp
--- a/Engine/Source/Editor/Kismet/Private/Debugging/SKismetDebuggingView.cpp
+++ b/Engine/Source/Editor/Kismet/Private/Debugging/SKismetDebuggingView.cpp
@@ -21,8 +21,8 @@ namespace KismetDebugViewConstants
 {
 	const FName ColumnId_Name( "Name" );
 	const FName ColumnId_Value( "Value" );
-	const FString ColumnText_Name( NSLOCTEXT("DebugViewUI", "Name", "Name").ToString() );
-	const FString ColumnText_Value( NSLOCTEXT("DebugViewUI", "Value", "Value").ToString() );
+	const FText ColumnText_Name( NSLOCTEXT("DebugViewUI", "Name", "Name") );
+	const FText ColumnText_Value( NSLOCTEXT("DebugViewUI", "Value", "Value") );
 }
 
 //////////////////////////////////////////////////////////////////////////
diff --git a/Engine/Source/Editor/Layers/Private/SLayersView.h b/Engine/Source/Editor/Layers/Private/SLayersView.h
--- a/Engine/Source/Editor/Layers/Private/SLayersView.h
+++ b/Engine/Source/Editor/Layers/Private/SLayersView.h
@@ -61,11 +61,11 @@ public:
 
 			/** LayerName label column */
 			+ SHeaderRow::Column( LayersView::ColumnID_LayerLabel )
-			.DefaultLabel( LOCTEXT("Column_LayerNameLabel", "Layer").ToString() )
+			.DefaultLabel( LOCTEXT("Column_LayerNameLabel", "Layer") )
 
 			/** Layer visibility column */
 			+ SHeaderRow::Column( LayersView::ColumnID_Visibility )
-			.DefaultLabel( NSLOCTEXT("LayersView", "Visibility", "Visibility").ToString() )
+			.DefaultLabel( NSLOCTEXT("LayersView", "Visibility", "Visibility") )
 			.FixedWidth( 48.0f );
 
 		ChildSlot
@@ -101,7 +101,7 @@ public:
 					.OnItemScrolledIntoView( this, &SLayersView::OnItemScrolledIntoView)
 
 					// Help text 
-					.ToolTipText(LOCTEXT("HelpText", "Drag actors from scene outliner or right click to add a new layer.").ToString())
+					.ToolTipText(LOCTEXT("HelpText", "Drag actors from scene outliner or right click to add a new layer."))
 				]
 			];
 
diff --git a/Engine/Source/Editor/PackagesDialog/Private/SPackagesDialog.cpp b/Engine/Source/Editor/PackagesDialog/Private/SPackagesDialog.cpp
--- a/Engine/Source/Editor/PackagesDialog/Private/SPackagesDialog.cpp
+++ b/Engine/Source/Editor/PackagesDialog/Private/SPackagesDialog.cpp
@@ -119,7 +119,7 @@ void SPackagesDialog::Construct(const FArguments& InArgs)
 
 	HeaderRowWidget->AddColumn(
 		SHeaderRow::Column( SPackagesDialogDefs::ColumnID_FileLabel )
-		.DefaultLabel( LOCTEXT("FileColumnLabel", "File" ).ToString() )
+		.DefaultLabel( LOCTEXT("FileColumnLabel", "File" ) )
 		.SortMode( this, &SPackagesDialog::GetColumnSortMode, SPackagesDialogDefs::ColumnID_FileLabel )
 		.OnSort( this, &SPackagesDialog::OnColumnSortModeChanged )
 		.FillWidth( 7.0f )
@@ -127,7 +127,7 @@ void SPackagesDialog::Construct(const FArguments& InArgs)
 
 	HeaderRowWidget->AddColumn(
 		SHeaderRow::Column( SPackagesDialogDefs::ColumnID_TypeLabel )
-		.DefaultLabel( LOCTEXT("TypeColumnLabel", "Type" ).ToString() )
+		.DefaultLabel( LOCTEXT("TypeColumnLabel", "Type" ) )
 		.SortMode( this, &SPackagesDialog::GetColumnSortMode, SPackagesDialogDefs::ColumnID_TypeLabel )
 		.OnSort( this, &SPackagesDialog::OnColumnSortModeChanged )
 		.FillWidth( 2.0f )
diff --git a/Engine/Source/Editor/Persona/Private/AnimGraphNodeDetails.cpp b/Engine/Source/Editor/Persona/Private/AnimGraphNodeDetails.cpp
--- a/Engine/Source/Editor/Persona/Private/AnimGraphNodeDetails.cpp
+++ b/Engine/Source/Editor/Persona/Private/AnimGraphNodeDetails.cpp
@@ -556,11 +556,11 @@ void FAnimGraphParentPlayerDetails::CustomizeDetails(class IDetailLayoutBuilder&
 			SNew(SHeaderRow)
 			+SHeaderRow::Column(FName("Name"))
 			.FillWidth(0.5f)
-			.DefaultLabel(LOCTEXT("ParentPlayer_NameCol", "Name").ToString())
+			.DefaultLabel(LOCTEXT("ParentPlayer_NameCol", "Name"))
 
 			+SHeaderRow::Column(FName("Asset"))
 			.FillWidth(0.5f)
-			.DefaultLabel(LOCTEXT("ParentPlayer_AssetCol", "Asset").ToString())
+			.DefaultLabel(LOCTEXT("ParentPlayer_AssetCol", "Asset"))
 		);
 
 	// Expand top level (blueprint) entries so the panel seems less empty
diff --git a/Engine/Source/Editor/Persona/Private/SMorphTargetViewer.cpp b/Engine/Source/Editor/Persona/Private/SMorphTargetViewer.cpp
--- a/Engine/Source/Editor/Persona/Private/SMorphTargetViewer.cpp
+++ b/Engine/Source/Editor/Persona/Private/SMorphTargetViewer.cpp
@@ -235,13 +235,13 @@ void SMorphTargetViewer::Construct(const FArguments& InArgs)
 			(
 				SNew( SHeaderRow )
 				+ SHeaderRow::Column( ColumnId_MorphTargetNameLabel )
-				.DefaultLabel( LOCTEXT( "MorphTargetNameLabel", "Morph Target Name" ).ToString() )
+				.DefaultLabel( LOCTEXT( "MorphTargetNameLabel", "Morph Target Name" ) )
 
 				+ SHeaderRow::Column( ColumnID_MorphTargetWeightLabel )
-				.DefaultLabel( LOCTEXT( "MorphTargetWeightLabel", "Weight" ).ToString() )
+				.DefaultLabel( LOCTEXT( "MorphTargetWeightLabel", "Weight" ) )
 
-				+ SHeaderRow::Column(ColumnID_MorphTargetVertCountLabel)
-				.DefaultLabel(LOCTEXT("MorphTargetVertCountLabel", "Vert Count").ToString())
+				+ SHeaderRow::Column( ColumnID_MorphTargetVertCountLabel )
+				.DefaultLabel( LOCTEXT("MorphTargetVertCountLabel", "Vert Count") )
 			)
 		]
 	];
diff --git a/Engine/Source/Editor/Persona/Private/SRetargetSourceWindow.cpp b/Engine/Source/Editor/Persona/Private/SRetargetSourceWindow.cpp
--- a/Engine/Source/Editor/Persona/Private/SRetargetSourceWindow.cpp
+++ b/Engine/Source/Editor/Persona/Private/SRetargetSourceWindow.cpp
@@ -237,10 +237,10 @@ void SRetargetSourceWindow::Construct(const FArguments& InArgs)
 			(
 				SNew( SHeaderRow )
 				+ SHeaderRow::Column( ColumnId_RetargetSourceNameLabel )
-				.DefaultLabel( LOCTEXT( "RetargetSourceNameLabel", "Retarget Source Name" ).ToString() )
+				.DefaultLabel( LOCTEXT( "RetargetSourceNameLabel", "Retarget Source Name" ) )
 
 				+ SHeaderRow::Column( ColumnID_BaseReferenceMeshLabel )
-				.DefaultLabel( LOCTEXT( "RetargetSourceWeightLabel", "Source Mesh" ).ToString() )
+				.DefaultLabel( LOCTEXT( "RetargetSourceWeightLabel", "Source Mesh" ) )
 			)
 		]
 	];
diff --git a/Engine/Source/Editor/Persona/Private/SRigWindow.cpp b/Engine/Source/Editor/Persona/Private/SRigWindow.cpp
--- a/Engine/Source/Editor/Persona/Private/SRigWindow.cpp
+++ b/Engine/Source/Editor/Persona/Private/SRigWindow.cpp
@@ -286,11 +286,11 @@ void SRigWindow::Construct(const FArguments& InArgs)
 			(
 				SNew( SHeaderRow )
 				+ SHeaderRow::Column( ColumnId_NodeNameLabel )
-				.DefaultLabel( LOCTEXT( "RigWindow_NodeNameLabel", "Node (Rig)" ).ToString() )
+				.DefaultLabel( LOCTEXT( "RigWindow_NodeNameLabel", "Node (Rig)" ) )
 				.FixedWidth(150.f)
 
 				+ SHeaderRow::Column( ColumnID_BoneNameLabel )
-				.DefaultLabel( LOCTEXT( "RigWindow_BoneNameLabel", "Bone (Skeleton)" ).ToString() )
+				.DefaultLabel( LOCTEXT( "RigWindow_BoneNameLabel", "Bone (Skeleton)" ) )
 			)
 		]
 	];
diff --git a/Engine/Source/Editor/Persona/Private/SSkeletonAnimNotifies.cpp b/Engine/Source/Editor/Persona/Private/SSkeletonAnimNotifies.cpp
--- a/Engine/Source/Editor/Persona/Private/SSkeletonAnimNotifies.cpp
+++ b/Engine/Source/Editor/Persona/Private/SSkeletonAnimNotifies.cpp
@@ -146,7 +146,7 @@ void SSkeletonAnimNotifies::Construct(const FArguments& InArgs)
 			(
 				SNew( SHeaderRow )
 				+ SHeaderRow::Column( ColumnId_AnimNotifyNameLabel )
-				.DefaultLabel( LOCTEXT( "AnimNotifyNameLabel", "Notify Name" ).ToString() )
+				.DefaultLabel( LOCTEXT( "AnimNotifyNameLabel", "Notify Name" ) )
 			)
 		]
 	];
diff --git a/Engine/Source/Editor/Persona/Private/SSkeletonSlotNames.cpp b/Engine/Source/Editor/Persona/Private/SSkeletonSlotNames.cpp
--- a/Engine/Source/Editor/Persona/Private/SSkeletonSlotNames.cpp
+++ b/Engine/Source/Editor/Persona/Private/SSkeletonSlotNames.cpp
@@ -190,7 +190,7 @@ void SSkeletonSlotNames::Construct(const FArguments& InArgs)
 			(
 				SNew( SHeaderRow )
 				+ SHeaderRow::Column( ColumnId_SlotNameLabel )
-				.DefaultLabel( LOCTEXT( "SlotNameNameLabel", "Slot Name" ).ToString() )
+				.DefaultLabel( LOCTEXT( "SlotNameNameLabel", "Slot Name" ) )
 			)
 		]
 	];
diff --git a/Engine/Source/Editor/Persona/Private/SSkeletonSmartNameManager.cpp b/Engine/Source/Editor/Persona/Private/SSkeletonSmartNameManager.cpp
--- a/Engine/Source/Editor/Persona/Private/SSkeletonSmartNameManager.cpp
+++ b/Engine/Source/Editor/Persona/Private/SSkeletonSmartNameManager.cpp
@@ -103,7 +103,7 @@ TSharedRef<SWidget> FSkeletonCurveNameManagerSummoner::CreateTabBody(const FWork
 	return SNew(SCurveNameManager)
 		.Persona(StaticCastSharedPtr<FPersona>(Toolkit))
 		.ContainerName(FName(USkeleton::AnimCurveMappingName))
-		.NameColumnDescription(LOCTEXT("CurveNameTabColumnDesc", "Name").ToString());
+		.NameColumnDescription(LOCTEXT("CurveNameTabColumnDesc", "Name"));
 }
 
 //////////////////////////////////////////////////////////////////////////
@@ -159,7 +159,7 @@ void SSkeletonSmartNameManager::Construct(const FArguments& InArgs)
 			(
 				SNew(SHeaderRow)
 				+SHeaderRow::Column(ColumnId_SmartNameLabel)
-				.DefaultLabel(FString(InArgs._NameColumnDescription))
+				.DefaultLabel(InArgs._NameColumnDescription)
 			)
 		]
 	];
diff --git a/Engine/Source/Editor/Persona/Private/SSkeletonSmartNameManager.h b/Engine/Source/Editor/Persona/Private/SSkeletonSmartNameManager.h
--- a/Engine/Source/Editor/Persona/Private/SSkeletonSmartNameManager.h
+++ b/Engine/Source/Editor/Persona/Private/SSkeletonSmartNameManager.h
@@ -64,7 +64,7 @@ public:
 
 	SLATE_ARGUMENT(TWeakPtr<FPersona>, Persona)			// Our Persona instance
 	SLATE_ARGUMENT(FName, ContainerName)				// The name of the smartname container to manage
-	SLATE_ARGUMENT(FString, NameColumnDescription)		// Description text for the display column header
+	SLATE_ARGUMENT(FText, NameColumnDescription)		// Description text for the display column header
 
 	SLATE_END_ARGS();
 
@@ -146,7 +146,7 @@ public:
 
 		SLATE_ARGUMENT(TWeakPtr<FPersona>, Persona)			// Our Persona instance
 		SLATE_ARGUMENT(FName, ContainerName)				// Name of the container to manage in the skeleton
-		SLATE_ARGUMENT(FString, NameColumnDescription)		// Description to show in the name column
+		SLATE_ARGUMENT(FText, NameColumnDescription)		// Description to show in the name column
 	SLATE_END_ARGS();
 
 	void Construct(const FArguments& InArgs)
diff --git a/Engine/Source/Editor/PropertyEditor/Private/SPropertyTreeViewImpl.cpp b/Engine/Source/Editor/PropertyEditor/Private/SPropertyTreeViewImpl.cpp
--- a/Engine/Source/Editor/PropertyEditor/Private/SPropertyTreeViewImpl.cpp
+++ b/Engine/Source/Editor/PropertyEditor/Private/SPropertyTreeViewImpl.cpp
@@ -243,7 +243,7 @@ void SPropertyTreeViewImpl::ConstructPropertyTree()
 						.BorderImage( FEditorStyle::GetBrush("NoBorder") )
 						[
 							SNew(STextBlock)
-							.Text( NSLOCTEXT("PropertyEditor", "NameColumn", "Name").ToString() )
+							.Text( NSLOCTEXT("PropertyEditor", "NameColumn", "Name") )
 						]
 					]
 					+SHeaderRow::Column(PropertyTreeConstants::ColumnId_Property)
@@ -254,7 +254,7 @@ void SPropertyTreeViewImpl::ConstructPropertyTree()
 						.BorderImage( FEditorStyle::GetBrush("NoBorder") )
 						[
 							SNew(STextBlock)
-							.Text( NSLOCTEXT("PropertyEditor", "PropertyColumn", "Value").ToString() )
+							.Text( NSLOCTEXT("PropertyEditor", "PropertyColumn", "Value") )
 						]
 					]
 				)
diff --git a/Engine/Source/Editor/PropertyEditor/Private/UserInterface/PropertyTree/PropertyTreeConstants.h b/Engine/Source/Editor/PropertyEditor/Private/UserInterface/PropertyTree/PropertyTreeConstants.h
--- a/Engine/Source/Editor/PropertyEditor/Private/UserInterface/PropertyTree/PropertyTreeConstants.h
+++ b/Engine/Source/Editor/PropertyEditor/Private/UserInterface/PropertyTree/PropertyTreeConstants.h
@@ -8,8 +8,8 @@ namespace PropertyTreeConstants
 	const FName ColumnId_Name( TEXT("Property") );
 	const FName ColumnId_Property( TEXT("Value") );
 
-	const FString ColumnText_Name( NSLOCTEXT("PropertyEditorWidgets", "Property", "Property").ToString() );
-	const FString ColumnText_Property( NSLOCTEXT("PropertyEditorWidgets", "Value", "Value").ToString() );
+	const FText ColumnText_Name( NSLOCTEXT("PropertyEditorWidgets", "Property", "Property") );
+	const FText ColumnText_Property( NSLOCTEXT("PropertyEditorWidgets", "Value", "Value") );
 }
 
 
diff --git a/Engine/Source/Editor/SceneOutliner/Private/SceneOutlinerItemLabelColumn.cpp b/Engine/Source/Editor/SceneOutliner/Private/SceneOutlinerItemLabelColumn.cpp
--- a/Engine/Source/Editor/SceneOutliner/Private/SceneOutlinerItemLabelColumn.cpp
+++ b/Engine/Source/Editor/SceneOutliner/Private/SceneOutlinerItemLabelColumn.cpp
@@ -644,7 +644,7 @@ FName FItemLabelColumn::GetColumnID()
 SHeaderRow::FColumn::FArguments FItemLabelColumn::ConstructHeaderRowColumn()
 {
 	return SHeaderRow::Column(GetColumnID())
-		.DefaultLabel(LOCTEXT("ItemLabel_HeaderText", "Label").ToString())
+		.DefaultLabel(LOCTEXT("ItemLabel_HeaderText", "Label"))
 		.FillWidth( 5.0f );
 }
 const TSharedRef<SWidget> FItemLabelColumn::ConstructRowWidget(FTreeItemRef TreeItem, const STableRow<FTreeItemPtr>& Row)
diff --git a/Engine/Source/Editor/SourceControlWindows/Private/SSourceControlHistory.cpp b/Engine/Source/Editor/SourceControlWindows/Private/SSourceControlHistory.cpp
--- a/Engine/Source/Editor/SourceControlWindows/Private/SSourceControlHistory.cpp
+++ b/Engine/Source/Editor/SourceControlWindows/Private/SSourceControlHistory.cpp
@@ -640,11 +640,11 @@ public:
 							.HeaderRow
 							(
 							SNew(SHeaderRow)
-							+ SHeaderRow::Column("Revision") .DefaultLabel(NSLOCTEXT("SourceControl.HistoryPanel.Header", "Revision", "Revision").ToString())			.FillWidth(100)
-							+ SHeaderRow::Column("Changelist") .DefaultLabel(NSLOCTEXT("SourceControl.HistoryPanel.Header", "Changelist", "ChangeList").ToString())		.FillWidth(150)
-							+ SHeaderRow::Column("Date") .DefaultLabel(NSLOCTEXT("SourceControl.HistoryPanel.Header", "Date", "Date Submitted").ToString())				.FillWidth(250)
-							+ SHeaderRow::Column("UserName") .DefaultLabel(NSLOCTEXT("SourceControl.HistoryPanel.Header", "UserName", "Submitted By").ToString())		.FillWidth(200)
-							+ SHeaderRow::Column("Description") .DefaultLabel(NSLOCTEXT("SourceControl.HistoryPanel.Header", "Description", "Description").ToString())	.FillWidth(300)
+							+ SHeaderRow::Column("Revision") .DefaultLabel(NSLOCTEXT("SourceControl.HistoryPanel.Header", "Revision", "Revision"))			.FillWidth(100)
+							+ SHeaderRow::Column("Changelist") .DefaultLabel(NSLOCTEXT("SourceControl.HistoryPanel.Header", "Changelist", "ChangeList"))	.FillWidth(150)
+							+ SHeaderRow::Column("Date") .DefaultLabel(NSLOCTEXT("SourceControl.HistoryPanel.Header", "Date", "Date Submitted"))			.FillWidth(250)
+							+ SHeaderRow::Column("UserName") .DefaultLabel(NSLOCTEXT("SourceControl.HistoryPanel.Header", "UserName", "Submitted By"))		.FillWidth(200)
+							+ SHeaderRow::Column("Description") .DefaultLabel(NSLOCTEXT("SourceControl.HistoryPanel.Header", "Description", "Description"))	.FillWidth(300)
 							)
 						]
 					]
diff --git a/Engine/Source/Editor/UnrealEd/Private/Dialogs/DlgSoundWaveOptions.cpp b/Engine/Source/Editor/UnrealEd/Private/Dialogs/DlgSoundWaveOptions.cpp
--- a/Engine/Source/Editor/UnrealEd/Private/Dialogs/DlgSoundWaveOptions.cpp
+++ b/Engine/Source/Editor/UnrealEd/Private/Dialogs/DlgSoundWaveOptions.cpp
@@ -264,26 +264,26 @@ void SSoundWaveCompressionOptions::Construct( const FArguments& InArgs )
 
 			// Quality label column
 			+ SHeaderRow::Column( SoundWaveCompressionOptions::ColumnID_QualityLabel )
-				.DefaultLabel( NSLOCTEXT("SoundWaveOptions", "QualityColumnLabel", "Quality").ToString() )
+				.DefaultLabel( NSLOCTEXT("SoundWaveOptions", "QualityColumnLabel", "Quality") )
 		
 			// Original Data Size label column
 			+ SHeaderRow::Column( SoundWaveCompressionOptions::ColumnID_OriginalDataSizeLabel )
-				.DefaultLabel( NSLOCTEXT("SoundWaveOptions", "OriginalDataSizeColumnLabel", "Original DataSize(Kb)").ToString() )
+				.DefaultLabel( NSLOCTEXT("SoundWaveOptions", "OriginalDataSizeColumnLabel", "Original DataSize(Kb)") )
 				.FillWidth( 3.0f )
 
 			// Vorbis Data Size label column
 			+ SHeaderRow::Column( SoundWaveCompressionOptions::ColumnID_VorbisDataSizeLabel )
-				.DefaultLabel( NSLOCTEXT("SoundWaveOptions", "VorbisDataSizeColumnLabel", "Vorbis DataSize(Kb)").ToString() )
+				.DefaultLabel( NSLOCTEXT("SoundWaveOptions", "VorbisDataSizeColumnLabel", "Vorbis DataSize(Kb)") )
 				.FillWidth( 3.0f )
 			
 			// XMA Data Size label column
 			+ SHeaderRow::Column( SoundWaveCompressionOptions::ColumnID_XMADataSizeLabel )
-				.DefaultLabel( NSLOCTEXT("SoundWaveOptions", "XMADataSizeColumnLabel", "XMA DataSize(Kb)").ToString() )
+				.DefaultLabel( NSLOCTEXT("SoundWaveOptions", "XMADataSizeColumnLabel", "XMA DataSize(Kb)") )
 				.FillWidth( 3.0f )
 			
 			// PS3 Data Size label column
 			+ SHeaderRow::Column( SoundWaveCompressionOptions::ColumnID_PS3DataSizeLabel )
-				.DefaultLabel( NSLOCTEXT("SoundWaveOptions", "PS3DataSizeColumnLabel", "PS3 DataSize(Kb)").ToString() )
+				.DefaultLabel( NSLOCTEXT("SoundWaveOptions", "PS3DataSizeColumnLabel", "PS3 DataSize(Kb)") )
 				.FillWidth( 3.0f );
 
 	this->ChildSlot
@@ -339,7 +339,7 @@ void SSoundWaveCompressionOptions::Construct( const FArguments& InArgs )
 						.OnClicked(this, &SSoundWaveCompressionOptions::OnOK_Clicked)
 						[
 							SNew(STextBlock)
-								.Text(NSLOCTEXT("SoundWaveOptions", "OK", "OK").ToString())
+								.Text(NSLOCTEXT("SoundWaveOptions", "OK", "OK"))
 						]
 				]
 
@@ -350,7 +350,7 @@ void SSoundWaveCompressionOptions::Construct( const FArguments& InArgs )
 						.OnClicked(this, &SSoundWaveCompressionOptions::OnCancel_Clicked)
 						[
 							SNew(STextBlock)
-							.Text(NSLOCTEXT("SoundWaveOptions", "Cancel", "Cancel").ToString())
+							.Text(NSLOCTEXT("SoundWaveOptions", "Cancel", "Cancel"))
 						]
 					]
 			]
diff --git a/Engine/Source/Editor/UnrealEd/Private/Dialogs/SDeleteAssetsDialog.cpp b/Engine/Source/Editor/UnrealEd/Private/Dialogs/SDeleteAssetsDialog.cpp
--- a/Engine/Source/Editor/UnrealEd/Private/Dialogs/SDeleteAssetsDialog.cpp
+++ b/Engine/Source/Editor/UnrealEd/Private/Dialogs/SDeleteAssetsDialog.cpp
@@ -166,19 +166,19 @@ TSharedRef<SWidget> SDeleteAssetsDialog::BuildDeleteDialog()
 	TSharedRef< SHeaderRow > HeaderRowWidget =
 		SNew( SHeaderRow )
 		+ SHeaderRow::Column( DeleteAssetsView::ColumnID_Asset )
-		.DefaultLabel( LOCTEXT( "Column_AssetName", "Asset" ).ToString() )
+		.DefaultLabel( LOCTEXT( "Column_AssetName", "Asset" ) )
 		.HAlignHeader( EHorizontalAlignment::HAlign_Left )
 		.FillWidth(0.5f)
 		+ SHeaderRow::Column( DeleteAssetsView::ColumnID_AssetClass )
-		.DefaultLabel( LOCTEXT( "Column_AssetClass", "Class" ).ToString() )
+		.DefaultLabel( LOCTEXT( "Column_AssetClass", "Class" ) )
 		.HAlignHeader( EHorizontalAlignment::HAlign_Left )
 		.FillWidth( 0.25f )
 		+ SHeaderRow::Column( DeleteAssetsView::ColumnID_DiskReferences )
-		.DefaultLabel( LOCTEXT( "Column_DiskReferences", "Asset Referencers" ).ToString() )
+		.DefaultLabel( LOCTEXT( "Column_DiskReferences", "Asset Referencers" ) )
 		.HAlignHeader( EHorizontalAlignment::HAlign_Left )
 		.FillWidth( 0.25f )
 		+ SHeaderRow::Column( DeleteAssetsView::ColumnID_MemoryReferences )
-		.DefaultLabel( LOCTEXT( "Column_MemoryReferences", "Memory References" ).ToString() )
+		.DefaultLabel( LOCTEXT( "Column_MemoryReferences", "Memory References" ) )
 		.HAlignHeader( EHorizontalAlignment::HAlign_Left )
 		.FillWidth( 0.25f );
 
diff --git a/Engine/Source/Editor/UnrealEd/Private/SSkeletonWidget.cpp b/Engine/Source/Editor/UnrealEd/Private/SSkeletonWidget.cpp
--- a/Engine/Source/Editor/UnrealEd/Private/SSkeletonWidget.cpp
+++ b/Engine/Source/Editor/UnrealEd/Private/SSkeletonWidget.cpp
@@ -226,9 +226,9 @@ void SSkeletonCompareWidget::SkeletonSelectionChanged(const FAssetData& AssetDat
 					(
 					SNew(SHeaderRow)
 					+SHeaderRow::Column(TEXT("Curretly Selected"))
-					.DefaultLabel(NSLOCTEXT("SkeletonWidget", "CurrentlySelected", "Currently Selected").ToString())
+					.DefaultLabel(NSLOCTEXT("SkeletonWidget", "CurrentlySelected", "Currently Selected"))
 					+SHeaderRow::Column(TEXT("Target Skeleton Bone"))
-					.DefaultLabel(NSLOCTEXT("SkeletonWidget", "TargetSkeletonBone", "Target Skeleton Bone").ToString())
+					.DefaultLabel(NSLOCTEXT("SkeletonWidget", "TargetSkeletonBone", "Target Skeleton Bone"))
 					)
 				]
 			];
