Fixed code relying on SLATE_TEXT_ATTRIBUTE for STextBlock.

UETOOL-213 - Minimize Slate FString -> FText conversion (remove SLATE_TEXT_ATTRIBUTE)

This fixes any editor/engine specific code that was passing text to Slate as FString rather than FText.

diff --git a/Engine/Plugins/2D/Paper2D/Source/Paper2DEditor/Private/FlipbookEditor/STimelineHeader.h b/Engine/Plugins/2D/Paper2D/Source/Paper2DEditor/Private/FlipbookEditor/STimelineHeader.h
--- a/Engine/Plugins/2D/Paper2D/Source/Paper2DEditor/Private/FlipbookEditor/STimelineHeader.h
+++ b/Engine/Plugins/2D/Paper2D/Source/Paper2DEditor/Private/FlipbookEditor/STimelineHeader.h
@@ -66,7 +66,7 @@ private:
 						.HAlign(HAlign_Center)
 						[
 							SNew(STextBlock)
-							.Text(FText::AsNumber(FrameIdx).ToString())
+							.Text(FText::AsNumber(FrameIdx))
 						]
 					];
 			}
diff --git a/Engine/Plugins/Developer/PerforceSourceControl/Source/PerforceSourceControl/Private/SPerforceSourceControlSettings.cpp b/Engine/Plugins/Developer/PerforceSourceControl/Source/PerforceSourceControl/Private/SPerforceSourceControlSettings.cpp
--- a/Engine/Plugins/Developer/PerforceSourceControl/Source/PerforceSourceControl/Private/SPerforceSourceControlSettings.cpp
+++ b/Engine/Plugins/Developer/PerforceSourceControl/Source/PerforceSourceControl/Private/SPerforceSourceControlSettings.cpp
@@ -429,7 +429,7 @@ TSharedRef<ITableRow> SPerforceSourceControlSettings::OnGenerateWorkspaceRow(TSh
 			.Padding(2.0f)
 			[
 				SNew(STextBlock)
-				.Text(*InItem)
+				.Text(FText::FromString(*InItem))
 				.Font(FEditorStyle::GetFontStyle(TEXT("PropertyWindow.NormalFont")))
 			]
 		];
@@ -444,9 +444,9 @@ void SPerforceSourceControlSettings::OnWorkspaceSelected(TSharedPtr<FString> InI
 	WorkspaceCombo->SetIsOpen(false);
 }
 
-FString SPerforceSourceControlSettings::OnGetButtonText() const
+FText SPerforceSourceControlSettings::OnGetButtonText() const
 {
-	return CurrentWorkspace;
+	return FText::FromString(CurrentWorkspace);
 }
 
 FReply SPerforceSourceControlSettings::OnCancelWorkspacesRequest()
diff --git a/Engine/Plugins/Developer/PerforceSourceControl/Source/PerforceSourceControl/Private/SPerforceSourceControlSettings.h b/Engine/Plugins/Developer/PerforceSourceControl/Source/PerforceSourceControl/Private/SPerforceSourceControlSettings.h
--- a/Engine/Plugins/Developer/PerforceSourceControl/Source/PerforceSourceControl/Private/SPerforceSourceControlSettings.h
+++ b/Engine/Plugins/Developer/PerforceSourceControl/Source/PerforceSourceControl/Private/SPerforceSourceControlSettings.h
@@ -96,7 +96,7 @@ private:
 	void OnWorkspaceSelected(TSharedPtr<FString> InItem, ESelectInfo::Type InSelectInfo);
 
 	/** Delegate for the auto workspaces dropdown button text */
-	FString OnGetButtonText() const;
+	FText OnGetButtonText() const;
 
 	/** Delegate that cancels the workspace request in progress */
 	FReply OnCancelWorkspacesRequest();
diff --git a/Engine/Plugins/Editor/EpicSurvey/Source/EpicSurvey/Private/SQuestionBlock.cpp b/Engine/Plugins/Editor/EpicSurvey/Source/EpicSurvey/Private/SQuestionBlock.cpp
--- a/Engine/Plugins/Editor/EpicSurvey/Source/EpicSurvey/Private/SQuestionBlock.cpp
+++ b/Engine/Plugins/Editor/EpicSurvey/Source/EpicSurvey/Private/SQuestionBlock.cpp
@@ -38,7 +38,7 @@ void SQuestionBlock::Construct( const FArguments& Args, const TSharedRef<FQuesti
 				.VAlign( VAlign_Center )
 				[
 					SNew( STextBlock )
-					.Text( Questions[QuestionIndex].ToString() )
+					.Text( Questions[QuestionIndex] )
 				]
 			];
 
@@ -55,7 +55,7 @@ void SQuestionBlock::Construct( const FArguments& Args, const TSharedRef<FQuesti
 					.Padding( FMargin( 0, 0, 10, 0 ) )
 					[
 						SNew( STextBlock )
-						.Text( Answers[AnswerIndex].Text.ToString() )
+						.Text( Answers[AnswerIndex].Text )
 					]
 				];
 				Grid->SetColumnFill( ColumnIndex, 1.0f );
@@ -79,7 +79,7 @@ void SQuestionBlock::Construct( const FArguments& Args, const TSharedRef<FQuesti
 					.VAlign( VAlign_Center )
 					[
 						SNew( STextBlock )
-						.Text( Questions[QuestionIndex].ToString() )
+						.Text( Questions[QuestionIndex] )
 					]
 				];
 			}
@@ -90,7 +90,7 @@ void SQuestionBlock::Construct( const FArguments& Args, const TSharedRef<FQuesti
 				.Padding( FMargin( 0, 5 ) )
 				[
 					SNew( STextBlock )
-					.Text( Questions[QuestionIndex].ToString() )
+					.Text( Questions[QuestionIndex] )
 				];
 			}
 
@@ -102,7 +102,7 @@ void SQuestionBlock::Construct( const FArguments& Args, const TSharedRef<FQuesti
 				.Padding( 0, 0, 3, 0 )
 				[
 					SNew( STextBlock )
-					.Text( Answers[AnswerIndex].Text.ToString() )
+					.Text( Answers[AnswerIndex].Text )
 				];
 
 				++ColumnIndex;
@@ -132,7 +132,7 @@ void SQuestionBlock::Construct( const FArguments& Args, const TSharedRef<FQuesti
 				.Padding( FMargin( 0, 0, 0, 8 ) )
 				[					
 					SNew( STextBlock )
-					.Text( Questions[QuestionIndex].ToString() )
+					.Text( Questions[QuestionIndex] )
 				]
 			];
 
@@ -154,7 +154,7 @@ void SQuestionBlock::Construct( const FArguments& Args, const TSharedRef<FQuesti
 						.OnCheckStateChanged( this, &SQuestionBlock::AnswerCheckStateChanged, TWeakPtr< FQuestionBlock >( Block ), QuestionIndex, AnswerIndex )
 						[
 							SNew( STextBlock )
-							.Text( Answers[AnswerIndex].Text.ToString() )
+							.Text( Answers[AnswerIndex].Text )
 						]
 					]
 				];
@@ -175,7 +175,7 @@ void SQuestionBlock::Construct( const FArguments& Args, const TSharedRef<FQuesti
 				.Padding( FMargin( 0, 0, 0, 8 ) )
 				[					
 					SNew( STextBlock )
-					.Text( Questions[QuestionIndex].ToString() )
+					.Text( Questions[QuestionIndex] )
 				]
 			];
 
@@ -193,7 +193,7 @@ void SQuestionBlock::Construct( const FArguments& Args, const TSharedRef<FQuesti
 					.VAlign( VAlign_Center )
 					[
 						SNew( STextBlock )
-						.Text( Answers[AnswerIndex].Text.ToString() )
+						.Text( Answers[AnswerIndex].Text )
 					]
 					+SHorizontalBox::Slot()
 					.Padding( FMargin( 0, 3, 10, 3 ) )
@@ -239,7 +239,7 @@ void SQuestionBlock::Construct( const FArguments& Args, const TSharedRef<FQuesti
 				.WidthOverride( QuestionColumnWidth )
 				[
 					SNew( STextBlock )
-					.Text( Questions[QuestionIndex].ToString() )
+					.Text( Questions[QuestionIndex] )
 				]
 			];
 
@@ -290,7 +290,7 @@ void SQuestionBlock::Construct( const FArguments& Args, const TSharedRef<FQuesti
 		.AutoHeight()
 		[
 			SNew( STextBlock )
-			.Text( Block->GetInstructions().ToString() )
+			.Text( Block->GetInstructions() )
 			.AutoWrapText( true )
 		];
 	}
diff --git a/Engine/Plugins/Editor/PluginsEditor/Source/PluginsEditor/Private/SPluginListTile.cpp b/Engine/Plugins/Editor/PluginsEditor/Source/PluginsEditor/Private/SPluginListTile.cpp
--- a/Engine/Plugins/Editor/PluginsEditor/Source/PluginsEditor/Private/SPluginListTile.cpp
+++ b/Engine/Plugins/Editor/PluginsEditor/Source/PluginsEditor/Private/SPluginListTile.cpp
@@ -101,7 +101,7 @@ void SPluginListTile::Construct( const FArguments& Args, const TSharedRef<SPlugi
 				.Padding(2.0f, 0.0f, 0.0f, 0.0f)
 				[
 					SNew(STextBlock)
-						.Text(Item->PluginStatus.CreatedBy)
+						.Text(FText::FromString(Item->PluginStatus.CreatedBy))
 				];
 		}
 		else
@@ -171,7 +171,7 @@ void SPluginListTile::Construct( const FArguments& Args, const TSharedRef<SPlugi
 										.Padding(PaddingAmount)
 										[
 											SNew(STextBlock)
-												.Text(Item->PluginStatus.FriendlyName)
+												.Text(FText::FromString(Item->PluginStatus.FriendlyName))
 												.HighlightText_Raw(&Owner->GetOwner().GetPluginTextFilter(), &FPluginTextFilter::GetRawFilterText)
 												.TextStyle(FPluginStyle::Get(), "PluginTile.NameText")
 										]
@@ -217,7 +217,7 @@ void SPluginListTile::Construct( const FArguments& Args, const TSharedRef<SPlugi
 													.Padding( 0.0f, 0.0f, 2.0f, 0.0f )	// Extra padding from the right edge
 													[
 														SNew(STextBlock)
-															.Text(Item->PluginStatus.VersionName)
+															.Text(FText::FromString(Item->PluginStatus.VersionName))
 															.TextStyle(FPluginStyle::Get(), "PluginTile.VersionNumberText")
 													]
 											]
@@ -232,7 +232,7 @@ void SPluginListTile::Construct( const FArguments& Args, const TSharedRef<SPlugi
 											.Padding( PaddingAmount )
 											[
 												SNew(STextBlock)
-													.Text(Item->PluginStatus.Description)
+													.Text(FText::FromString(Item->PluginStatus.Description))
 													.AutoWrapText(true)
 											]
 
diff --git a/Engine/Plugins/Experimental/CodeView/Source/CodeView/Private/SCodeView.cpp b/Engine/Plugins/Experimental/CodeView/Source/CodeView/Private/SCodeView.cpp
--- a/Engine/Plugins/Experimental/CodeView/Source/CodeView/Private/SCodeView.cpp
+++ b/Engine/Plugins/Experimental/CodeView/Source/CodeView/Private/SCodeView.cpp
@@ -450,14 +450,14 @@ namespace CodeView
 	}
 
 
-	FString SCodeView::GetCustomColumnTextForTreeItem( FTreeItemPtr TreeItem ) const
+	FText SCodeView::GetCustomColumnTextForTreeItem( FTreeItemPtr TreeItem ) const
 	{
-		FString CustomColumnText;
+		FText CustomColumnText = FText::GetEmpty();
 
 		switch( CurrentCustomColumnMode )
 		{
 			case ECustomColumnMode::ModuleName:
-				CustomColumnText = TreeItem->ModuleName;
+				CustomColumnText = FText::FromString(TreeItem->ModuleName);
 				break;
 		};
 
@@ -513,7 +513,7 @@ namespace CodeView
 						[
 							SNew( STextBlock )
 
- 								.Text( Item->GetAsString() ) 
+ 								.Text( FText::FromString(Item->GetAsString()) ) 
 
 								// Bind our filter text as the highlight string for the text block, so that when the user
 								// starts typing search criteria, this text highlights
diff --git a/Engine/Plugins/Experimental/CodeView/Source/CodeView/Private/SCodeView.h b/Engine/Plugins/Experimental/CodeView/Source/CodeView/Private/SCodeView.h
--- a/Engine/Plugins/Experimental/CodeView/Source/CodeView/Private/SCodeView.h
+++ b/Engine/Plugins/Experimental/CodeView/Source/CodeView/Private/SCodeView.h
@@ -162,7 +162,7 @@ namespace CodeView
 		void Populate();
 
 		/** Gets text for the specified item to display in the custom column for the tree */
-		FString GetCustomColumnTextForTreeItem( FTreeItemPtr TreeItem ) const;
+		FText GetCustomColumnTextForTreeItem( FTreeItemPtr TreeItem ) const;
 
 		/** Called by STreeView to generate a table row for the specified item */
 		TSharedRef< ITableRow > OnGenerateRowForTree( FTreeItemPtr Item, const TSharedRef< STableViewBase >& OwnerTable );
diff --git a/Engine/Plugins/Experimental/LiveEditor/Source/LiveEditor/Private/K2Node_LiveEditObject.cpp b/Engine/Plugins/Experimental/LiveEditor/Source/LiveEditor/Private/K2Node_LiveEditObject.cpp
--- a/Engine/Plugins/Experimental/LiveEditor/Source/LiveEditor/Private/K2Node_LiveEditObject.cpp
+++ b/Engine/Plugins/Experimental/LiveEditor/Source/LiveEditor/Private/K2Node_LiveEditObject.cpp
@@ -783,7 +783,7 @@ void UK2Node_LiveEditObject::SetPinToolTip(UEdGraphPin& MutatablePin, const FTex
 	if (K2Schema != nullptr)
 	{
 		MutatablePin.PinToolTip += TEXT(" ");
-		MutatablePin.PinToolTip += K2Schema->GetPinDisplayName(&MutatablePin);
+		MutatablePin.PinToolTip += K2Schema->GetPinDisplayName(&MutatablePin).ToString();
 	}
 
 	MutatablePin.PinToolTip += FString(TEXT("\n")) + PinDescription.ToString();
diff --git a/Engine/Plugins/Experimental/LiveEditor/Source/LiveEditor/Private/LiveEditorConfigWindow.cpp b/Engine/Plugins/Experimental/LiveEditor/Source/LiveEditor/Private/LiveEditorConfigWindow.cpp
--- a/Engine/Plugins/Experimental/LiveEditor/Source/LiveEditor/Private/LiveEditorConfigWindow.cpp
+++ b/Engine/Plugins/Experimental/LiveEditor/Source/LiveEditor/Private/LiveEditorConfigWindow.cpp
@@ -227,7 +227,7 @@ public:
 				.Padding(2.0f)
 				[
 					SNew( STextBlock )
-					.Text( DeviceData.DeviceName )
+					.Text( FText::FromString(DeviceData.DeviceName) )
 					.ColorAndOpacity( FLinearColor::White )
 				]
 				+SVerticalBox::Slot()
diff --git a/Engine/Plugins/Experimental/LiveEditor/Source/LiveEditor/Private/LiveEditorConnectionWindow.cpp b/Engine/Plugins/Experimental/LiveEditor/Source/LiveEditor/Private/LiveEditorConnectionWindow.cpp
--- a/Engine/Plugins/Experimental/LiveEditor/Source/LiveEditor/Private/LiveEditorConnectionWindow.cpp
+++ b/Engine/Plugins/Experimental/LiveEditor/Source/LiveEditor/Private/LiveEditorConnectionWindow.cpp
@@ -47,10 +47,10 @@ public:
 		return FReply::Handled();
 	}
 
-	FString GetConnectionInfo() const
+	FText GetConnectionInfo() const
 	{
 		check( ConnectionInfo.IsValid() );
-		return *ConnectionInfo.Get();
+		return FText::FromString(*ConnectionInfo.Get());
 	}
 
 private:
diff --git a/Engine/Plugins/Messaging/MessagingDebugger/Source/MessagingDebugger/Private/Widgets/EndpointDetails/SMessagingAddressTableRow.h b/Engine/Plugins/Messaging/MessagingDebugger/Source/MessagingDebugger/Private/Widgets/EndpointDetails/SMessagingAddressTableRow.h
--- a/Engine/Plugins/Messaging/MessagingDebugger/Source/MessagingDebugger/Private/Widgets/EndpointDetails/SMessagingAddressTableRow.h
+++ b/Engine/Plugins/Messaging/MessagingDebugger/Source/MessagingDebugger/Private/Widgets/EndpointDetails/SMessagingAddressTableRow.h
@@ -46,6 +46,10 @@ public:
 	BEGIN_SLATE_FUNCTION_BUILD_OPTIMIZATION
 	virtual TSharedRef<SWidget> GenerateWidgetForColumn( const FName& ColumnName ) override
 	{
+		static const FNumberFormattingOptions TimeRegisteredFormattingOptions = FNumberFormattingOptions()
+			.SetMinimumFractionalDigits(5)
+			.SetMaximumFractionalDigits(5);
+
 		if (ColumnName == "Address")
 		{
 			return SNew(SBox)
@@ -53,7 +57,7 @@ public:
 				.VAlign(VAlign_Center)
 				[
 					SNew(STextBlock)
-						.Text(AddressInfo->Address.ToString())
+						.Text(FText::FromString(AddressInfo->Address.ToString()))
 				];
 		}
 		else if (ColumnName == "TimeRegistered")
@@ -63,7 +67,7 @@ public:
 				.VAlign(VAlign_Center)
 				[
 					SNew(STextBlock)
-						.Text(FString::Printf(TEXT("%.5f"), AddressInfo->TimeRegistered - GStartTime))
+						.Text(FText::AsNumber(AddressInfo->TimeRegistered - GStartTime, &TimeRegisteredFormattingOptions))
 				];
 		}
 		else if (ColumnName == "TimeUnregistered")
@@ -73,7 +77,7 @@ public:
 				.VAlign(VAlign_Center)
 				[
 					SNew(STextBlock)
-						.Text((AddressInfo->TimeUnregistered > 0.0) ? FString::Printf(TEXT("%.5f"), AddressInfo->TimeUnregistered - GStartTime) : TEXT("Never"))
+						.Text((AddressInfo->TimeUnregistered > 0.0) ? FText::AsNumber(AddressInfo->TimeUnregistered - GStartTime, &TimeRegisteredFormattingOptions) : LOCTEXT("Never", "Never"))
 				];
 		}
 
diff --git a/Engine/Plugins/Messaging/MessagingDebugger/Source/MessagingDebugger/Private/Widgets/EndpointDetails/SMessagingEndpointDetails.cpp b/Engine/Plugins/Messaging/MessagingDebugger/Source/MessagingDebugger/Private/Widgets/EndpointDetails/SMessagingEndpointDetails.cpp
--- a/Engine/Plugins/Messaging/MessagingDebugger/Source/MessagingDebugger/Private/Widgets/EndpointDetails/SMessagingEndpointDetails.cpp
+++ b/Engine/Plugins/Messaging/MessagingDebugger/Source/MessagingDebugger/Private/Widgets/EndpointDetails/SMessagingEndpointDetails.cpp
@@ -138,29 +138,29 @@ TSharedRef<ITableRow> SMessagingEndpointDetails::HandleAddressListGenerateRow( F
 }
 
 
-FString SMessagingEndpointDetails::HandleEndpointDetailsReceivedMessagesText() const
+FText SMessagingEndpointDetails::HandleEndpointDetailsReceivedMessagesText() const
 {
 	FMessageTracerEndpointInfoPtr SelectedEndpoint = Model->GetSelectedEndpoint();
 
 	if (SelectedEndpoint.IsValid())
 	{
-		return FString::Printf(TEXT("%i"), SelectedEndpoint->ReceivedMessages.Num());
+		return FText::AsNumber(SelectedEndpoint->ReceivedMessages.Num());
 	}
 
-	return FString();
+	return FText::GetEmpty();
 }
 
 
-FString SMessagingEndpointDetails::HandleEndpointDetailsSentMessagesText() const
+FText SMessagingEndpointDetails::HandleEndpointDetailsSentMessagesText() const
 {
 	FMessageTracerEndpointInfoPtr SelectedEndpoint = Model->GetSelectedEndpoint();
 
 	if (SelectedEndpoint.IsValid())
 	{
-		return FString::Printf(TEXT("%i"), SelectedEndpoint->SentMessages.Num());
+		return FText::AsNumber(SelectedEndpoint->SentMessages.Num());
 	}
 
-	return FString();
+	return FText::GetEmpty();
 }
 
 
diff --git a/Engine/Plugins/Messaging/MessagingDebugger/Source/MessagingDebugger/Private/Widgets/EndpointDetails/SMessagingEndpointDetails.h b/Engine/Plugins/Messaging/MessagingDebugger/Source/MessagingDebugger/Private/Widgets/EndpointDetails/SMessagingEndpointDetails.h
--- a/Engine/Plugins/Messaging/MessagingDebugger/Source/MessagingDebugger/Private/Widgets/EndpointDetails/SMessagingEndpointDetails.h
+++ b/Engine/Plugins/Messaging/MessagingDebugger/Source/MessagingDebugger/Private/Widgets/EndpointDetails/SMessagingEndpointDetails.h
@@ -42,10 +42,10 @@ private:
 	TSharedRef<ITableRow> HandleAddressListGenerateRow( FMessageTracerAddressInfoPtr AddressInfo, const TSharedRef<STableViewBase>& OwnerTable );
 
 	/** Callback for getting the number of received messages. */
-	FString HandleEndpointDetailsReceivedMessagesText() const;
+	FText HandleEndpointDetailsReceivedMessagesText() const;
 
 	/** Callback for getting the number of sent messages. */
-	FString HandleEndpointDetailsSentMessagesText() const;
+	FText HandleEndpointDetailsSentMessagesText() const;
 
 private:
 
diff --git a/Engine/Plugins/Messaging/MessagingDebugger/Source/MessagingDebugger/Private/Widgets/Endpoints/SMessagingEndpointsTableRow.h b/Engine/Plugins/Messaging/MessagingDebugger/Source/MessagingDebugger/Private/Widgets/Endpoints/SMessagingEndpointsTableRow.h
--- a/Engine/Plugins/Messaging/MessagingDebugger/Source/MessagingDebugger/Private/Widgets/Endpoints/SMessagingEndpointsTableRow.h
+++ b/Engine/Plugins/Messaging/MessagingDebugger/Source/MessagingDebugger/Private/Widgets/Endpoints/SMessagingEndpointsTableRow.h
@@ -77,7 +77,7 @@ public:
 					[
 						SNew(STextBlock)
 							.HighlightText(HighlightText)
-							.Text(EndpointInfo->Name.ToString())
+							.Text(FText::FromName(EndpointInfo->Name))
 					];
 		}
 		else if (ColumnName == "Messages")
@@ -113,15 +113,15 @@ public:
 private:
 
 	/** Gets the text for the Messages column. */
-	FString HandleMessagesText() const
+	FText HandleMessagesText() const
 	{
-		return FString::Printf(TEXT("%i"), EndpointInfo->ReceivedMessages.Num() + EndpointInfo->SentMessages.Num());
+		return FText::AsNumber(EndpointInfo->ReceivedMessages.Num() + EndpointInfo->SentMessages.Num());
 	}
 
 	/** Gets the tooltip text for the Messages column. */
-	FString HandleMessagesTooltipText() const
+	FText HandleMessagesTooltipText() const
 	{
-		return FString::Printf(TEXT("In: %i\nOut: %i"), EndpointInfo->ReceivedMessages.Num(), EndpointInfo->SentMessages.Num());
+		return FText::Format(LOCTEXT("MessagesTooltipTextFmt", "In: {0}\nOut: {1}"), FText::AsNumber(EndpointInfo->ReceivedMessages.Num()), FText::AsNumber(EndpointInfo->SentMessages.Num()));
 	}
 
 	/** Handles changing the checked state of the visibility check box. */
diff --git a/Engine/Plugins/Messaging/MessagingDebugger/Source/MessagingDebugger/Private/Widgets/History/SMessagingHistoryTableRow.h b/Engine/Plugins/Messaging/MessagingDebugger/Source/MessagingDebugger/Private/Widgets/History/SMessagingHistoryTableRow.h
--- a/Engine/Plugins/Messaging/MessagingDebugger/Source/MessagingDebugger/Private/Widgets/History/SMessagingHistoryTableRow.h
+++ b/Engine/Plugins/Messaging/MessagingDebugger/Source/MessagingDebugger/Private/Widgets/History/SMessagingHistoryTableRow.h
@@ -111,7 +111,7 @@ public:
 					SNew(STextBlock)
 						.ColorAndOpacity(this, &SMessagingHistoryTableRow::HandleTextColorAndOpacity)
 						.HighlightText(HighlightText)
-						.Text(MessageInfo->Context->GetMessageType().ToString())
+						.Text(FText::FromName(MessageInfo->Context->GetMessageType()))
 				];
 		}
 		else if (ColumnName == "Recipients")
diff --git a/Engine/Plugins/Messaging/MessagingDebugger/Source/MessagingDebugger/Private/Widgets/MessageDetails/SMessagingDispatchStateTableRow.h b/Engine/Plugins/Messaging/MessagingDebugger/Source/MessagingDebugger/Private/Widgets/MessageDetails/SMessagingDispatchStateTableRow.h
--- a/Engine/Plugins/Messaging/MessagingDebugger/Source/MessagingDebugger/Private/Widgets/MessageDetails/SMessagingDispatchStateTableRow.h
+++ b/Engine/Plugins/Messaging/MessagingDebugger/Source/MessagingDebugger/Private/Widgets/MessageDetails/SMessagingDispatchStateTableRow.h
@@ -99,7 +99,7 @@ public:
 				.VAlign(VAlign_Center)
 				[
 					SNew(STextBlock)
-						.Text(EndpointInfo->Name.ToString())
+						.Text(FText::FromName(EndpointInfo->Name))
 				];
 		}
 
diff --git a/Engine/Plugins/Messaging/MessagingDebugger/Source/MessagingDebugger/Private/Widgets/Types/SMessagingTypesTableRow.h b/Engine/Plugins/Messaging/MessagingDebugger/Source/MessagingDebugger/Private/Widgets/Types/SMessagingTypesTableRow.h
--- a/Engine/Plugins/Messaging/MessagingDebugger/Source/MessagingDebugger/Private/Widgets/Types/SMessagingTypesTableRow.h
+++ b/Engine/Plugins/Messaging/MessagingDebugger/Source/MessagingDebugger/Private/Widgets/Types/SMessagingTypesTableRow.h
@@ -76,7 +76,7 @@ public:
 				[
 					SNew(STextBlock)
 						.HighlightText(HighlightText)
-						.Text(TypeInfo->TypeName.ToString())
+						.Text(FText::FromName(TypeInfo->TypeName))
 				];
 		}
 		else if (ColumnName == "Visibility")
@@ -100,9 +100,9 @@ public:
 private:
 
 	/** Gets the text for the Messages column. */
-	FString HandleMessagesText() const
+	FText HandleMessagesText() const
 	{
-		return FString::Printf(TEXT("%i"), TypeInfo->Messages.Num());
+		return FText::AsNumber(TypeInfo->Messages.Num());
 	}
 
 	/** Handles changing the checked state of the visibility check box. */
diff --git a/Engine/Source/Developer/AssetTools/Private/AssetTypeActions/AssetTypeActions_SkeletalMesh.cpp b/Engine/Source/Developer/AssetTools/Private/AssetTypeActions/AssetTypeActions_SkeletalMesh.cpp
--- a/Engine/Source/Developer/AssetTools/Private/AssetTypeActions/AssetTypeActions_SkeletalMesh.cpp
+++ b/Engine/Source/Developer/AssetTools/Private/AssetTypeActions/AssetTypeActions_SkeletalMesh.cpp
@@ -194,7 +194,7 @@ public:
 			.IsChecked( this, &SDlgMergeSkeleton::IsCheckboxChecked, ButtonId )
 			.OnCheckStateChanged( this, &SDlgMergeSkeleton::OnCheckboxChanged, ButtonId )
 			[
-				SNew(STextBlock).Text(Label)
+				SNew(STextBlock).Text(FText::FromString(Label))
 			];
 	}
 
diff --git a/Engine/Source/Developer/AssetTools/Private/AssetTypeActions/AssetTypeActions_Skeleton.cpp b/Engine/Source/Developer/AssetTools/Private/AssetTypeActions/AssetTypeActions_Skeleton.cpp
--- a/Engine/Source/Developer/AssetTools/Private/AssetTypeActions/AssetTypeActions_Skeleton.cpp
+++ b/Engine/Source/Developer/AssetTools/Private/AssetTypeActions/AssetTypeActions_Skeleton.cpp
@@ -193,7 +193,7 @@ public:
 			.IsChecked( this, &SCreateRigDlg::IsCheckboxChecked, ButtonId )
 			.OnCheckStateChanged( this, &SCreateRigDlg::OnCheckboxChanged, ButtonId )
 			[
-				SNew(STextBlock).Text(Label)
+				SNew(STextBlock).Text(FText::FromString(Label))
 			];
 	}
 
diff --git a/Engine/Source/Developer/AssetTools/Private/SPackageReportDialog.cpp b/Engine/Source/Developer/AssetTools/Private/SPackageReportDialog.cpp
--- a/Engine/Source/Developer/AssetTools/Private/SPackageReportDialog.cpp
+++ b/Engine/Source/Developer/AssetTools/Private/SPackageReportDialog.cpp
@@ -211,7 +211,7 @@ TSharedRef<ITableRow> SPackageReportDialog::GenerateTreeRow( TSharedPtr<FPackage
 			+SHorizontalBox::Slot()
 			.FillWidth(1.f)
 			[
-				SNew(STextBlock).Text(TreeItem->NodeName)
+				SNew(STextBlock).Text(FText::FromString(TreeItem->NodeName))
 			]
 		];
 }
diff --git a/Engine/Source/Developer/AutomationWindow/Private/AutomationPresetManager.cpp b/Engine/Source/Developer/AutomationWindow/Private/AutomationPresetManager.cpp
--- a/Engine/Source/Developer/AutomationWindow/Private/AutomationPresetManager.cpp
+++ b/Engine/Source/Developer/AutomationWindow/Private/AutomationPresetManager.cpp
@@ -16,7 +16,7 @@ AutomationPresetRef FAutomationTestPresetManager::AddNewPreset( )
 }
 
 
-AutomationPresetRef FAutomationTestPresetManager::AddNewPreset( const FString& PresetName, const TArray<FString>& SelectedTests )
+AutomationPresetRef FAutomationTestPresetManager::AddNewPreset( const FText& PresetName, const TArray<FString>& SelectedTests )
 {
 	AutomationPresetRef NewPreset = AddNewPreset();
 	NewPreset->SetPresetName(PresetName);
diff --git a/Engine/Source/Developer/AutomationWindow/Private/AutomationPresetManager.h b/Engine/Source/Developer/AutomationWindow/Private/AutomationPresetManager.h
--- a/Engine/Source/Developer/AutomationWindow/Private/AutomationPresetManager.h
+++ b/Engine/Source/Developer/AutomationWindow/Private/AutomationPresetManager.h
@@ -45,7 +45,7 @@ public:
 	 * @return The preset's name.
 	 * @see GetID, SetPresetName
 	 */
-	const FString& GetPresetName( ) const
+	const FText& GetPresetName( ) const
 	{
 		return PresetName;
 	}
@@ -56,7 +56,7 @@ public:
 	 * @param InPresetName The name to set.
 	 * @see GetPresetName
 	 */
-	void SetPresetName(const FString& InPresetName)
+	void SetPresetName(const FText& InPresetName)
 	{
 		PresetName = InPresetName;
 	}
@@ -112,7 +112,7 @@ private:
 	FGuid Id;
 
 	/** The name of this preset. */
-	FString PresetName;
+	FText PresetName;
 
 	/** The list of enabled test names. */
 	TArray<FString> EnabledTests;
@@ -139,7 +139,7 @@ public:
 	 * @param PresetName The name of the new preset.
 	 * @param SelectedTests The list of enabled tests.
 	 */
-	virtual AutomationPresetRef AddNewPreset( const FString& PresetName, const TArray<FString>& SelectedTests );
+	virtual AutomationPresetRef AddNewPreset( const FText& PresetName, const TArray<FString>& SelectedTests );
 
 	/**
 	 * Returns a reference to the list that holds the presets.
diff --git a/Engine/Source/Developer/AutomationWindow/Private/SAutomationGraphicalResultBox.cpp b/Engine/Source/Developer/AutomationWindow/Private/SAutomationGraphicalResultBox.cpp
--- a/Engine/Source/Developer/AutomationWindow/Private/SAutomationGraphicalResultBox.cpp
+++ b/Engine/Source/Developer/AutomationWindow/Private/SAutomationGraphicalResultBox.cpp
@@ -182,7 +182,7 @@ void SAutomationGraphicalResultBox::CreateWidgets()
 					+SVerticalBox::Slot()
 					[
 						SNew(STextBlock)
-						.Text(DeviceIt->InstanceName)
+						.Text(FText::FromString(DeviceIt->InstanceName))
 					]
 					+SVerticalBox::Slot()
 					[
diff --git a/Engine/Source/Developer/AutomationWindow/Private/SAutomationTestItem.cpp b/Engine/Source/Developer/AutomationWindow/Private/SAutomationTestItem.cpp
--- a/Engine/Source/Developer/AutomationWindow/Private/SAutomationTestItem.cpp
+++ b/Engine/Source/Developer/AutomationWindow/Private/SAutomationTestItem.cpp
@@ -164,7 +164,7 @@ TSharedRef<SWidget> SAutomationTestItem::GenerateWidgetForColumn( const FName& C
 				//name of the test
 				SNew( STextBlock )
 				.HighlightText( HighlightText )
-				.Text( TestStatus->GetDisplayNameWithDecoration() )
+				.Text( FText::FromString(TestStatus->GetDisplayNameWithDecoration()) )
 			];
 	}
 	else if( ColumnName == AutomationTestWindowConstants::SmokeTest )
diff --git a/Engine/Source/Developer/AutomationWindow/Private/SAutomationTestTreeView.h b/Engine/Source/Developer/AutomationWindow/Private/SAutomationTestTreeView.h
--- a/Engine/Source/Developer/AutomationWindow/Private/SAutomationTestTreeView.h
+++ b/Engine/Source/Developer/AutomationWindow/Private/SAutomationTestTreeView.h
@@ -116,7 +116,7 @@ public:
 				.VAlign(VAlign_Center)
 				[
 					SNew(STextBlock)
-					.Text(ErrorString)
+					.Text(FText::FromString(ErrorString))
 				];
 			}
 			else
diff --git a/Engine/Source/Developer/AutomationWindow/Private/SAutomationWindow.cpp b/Engine/Source/Developer/AutomationWindow/Private/SAutomationWindow.cpp
--- a/Engine/Source/Developer/AutomationWindow/Private/SAutomationWindow.cpp
+++ b/Engine/Source/Developer/AutomationWindow/Private/SAutomationWindow.cpp
@@ -832,7 +832,7 @@ void SAutomationWindow::HandlePresetTextCommited( const FText& CommittedText, ET
 
 		TArray<FString> EnabledTests;
 		AutomationController->GetEnabledTestNames(EnabledTests);
-		AutomationPresetPtr NewPreset = TestPresetManager->AddNewPreset(CommittedText.ToString(),EnabledTests);
+		AutomationPresetPtr NewPreset = TestPresetManager->AddNewPreset(CommittedText,EnabledTests);
 		PresetComboBox->SetSelectedItem(NewPreset);
 		SelectedPreset = NewPreset;
 
@@ -911,15 +911,15 @@ FReply SAutomationWindow::HandleRemovePresetClicked()
 	return FReply::Handled();
 }
 
-FString SAutomationWindow::GetPresetComboText() const
+FText SAutomationWindow::GetPresetComboText() const
 {
 	if( SelectedPreset.IsValid() )
 	{
-		return *SelectedPreset->GetPresetName();
+		return SelectedPreset->GetPresetName();
 	}
 	else
 	{
-		return LOCTEXT("AutomationPresetComboLabel", "<Select A Preset>").ToString();
+		return LOCTEXT("AutomationPresetComboLabel", "<Select A Preset>");
 	}
 }
 
@@ -1403,29 +1403,29 @@ TSharedRef<ITableRow> SAutomationWindow::OnGenerateWidgetForLog(TSharedPtr<FAuto
 			[
 				SNew(STextBlock)
 				.TextStyle( FEditorStyle::Get(), Message->Style )
-				.Text(FString::Printf(TEXT("%s"), *Message->Text))
+				.Text(FText::FromString(Message->Text))
 			]
 		];
 }
 
 
-FString SAutomationWindow::OnGetNumEnabledTestsString() const
+FText SAutomationWindow::OnGetNumEnabledTestsString() const
 {
 	int32 NumPasses = AutomationController->GetNumPasses();
 	if( NumPasses > 1 )
 	{
-		return FString::Printf(TEXT("%d x%d"), AutomationController->GetEnabledTestsNum(),NumPasses);
+		return FText::Format(LOCTEXT("NumEnabledTestsFmt", "{0} x{1}"), FText::AsNumber(AutomationController->GetEnabledTestsNum()), FText::AsNumber(NumPasses));
 	}
 	else
 	{
-		return FString::Printf(TEXT("%d"), AutomationController->GetEnabledTestsNum());
+		return FText::AsNumber(AutomationController->GetEnabledTestsNum());
 	}
 }
 
 
-FString SAutomationWindow::OnGetNumDevicesInClusterString(const int32 ClusterIndex) const
+FText SAutomationWindow::OnGetNumDevicesInClusterString(const int32 ClusterIndex) const
 {
-	return FString::Printf(TEXT("%d"), AutomationController->GetNumDevicesInCluster(ClusterIndex));
+	return FText::AsNumber(AutomationController->GetNumDevicesInCluster(ClusterIndex));
 }
 
 
@@ -1729,15 +1729,15 @@ const FSlateBrush* SAutomationWindow::GetRunAutomationIcon() const
 }
 
 
-FString SAutomationWindow::GetRunAutomationLabel() const
+FText SAutomationWindow::GetRunAutomationLabel() const
 {
 	if( AutomationControllerState == EAutomationControllerModuleState::Running )
 	{
-		return LOCTEXT( "RunStopTestsLabel", "Stop Tests" ).ToString();
+		return LOCTEXT( "RunStopTestsLabel", "Stop Tests" );
 	}
 	else
 	{
-		return LOCTEXT( "RunStartTestsLabel", "Start Tests" ).ToString();
+		return LOCTEXT( "RunStartTestsLabel", "Start Tests" );
 	}
 }
 
diff --git a/Engine/Source/Developer/AutomationWindow/Private/SAutomationWindow.h b/Engine/Source/Developer/AutomationWindow/Private/SAutomationWindow.h
--- a/Engine/Source/Developer/AutomationWindow/Private/SAutomationWindow.h
+++ b/Engine/Source/Developer/AutomationWindow/Private/SAutomationWindow.h
@@ -213,10 +213,10 @@ private:
 	TSharedRef<ITableRow> OnGenerateWidgetForLog(TSharedPtr<FAutomationOutputMessage> Message, const TSharedRef<STableViewBase>& OwnerTable);
 
 	/** Returns number of enabled tests (regardless of visibility) */
-	FString OnGetNumEnabledTestsString() const;
+	FText OnGetNumEnabledTestsString() const;
 	
 	/** Returns number of workers in a device cluster */
-	FString OnGetNumDevicesInClusterString(const int32 ClusterIndex) const;
+	FText OnGetNumDevicesInClusterString(const int32 ClusterIndex) const;
 
 	/** Callback when the list has been refreshed by the Automation Controller */
 	void OnRefreshTestCallback();
@@ -330,7 +330,7 @@ private:
 	*
 	* @return A start or stop tests label Depending on state.
 	*/
-	FString GetRunAutomationLabel() const;
+	FText GetRunAutomationLabel() const;
 
 	/**
 	* Gets the brush to use for the test list background.
@@ -426,7 +426,7 @@ private:
 	void ExpandEnabledTests( TSharedPtr< IAutomationReport > InReport );
 
 	/** Gets the text to display for the preset combo box. */
-	FString GetPresetComboText() const;
+	FText GetPresetComboText() const;
 
 	/**
 	 * Handle the copy button clicked in the command bar.
diff --git a/Engine/Source/Developer/CollisionAnalyzer/Private/SCAQueryDetails.cpp b/Engine/Source/Developer/CollisionAnalyzer/Private/SCAQueryDetails.cpp
--- a/Engine/Source/Developer/CollisionAnalyzer/Private/SCAQueryDetails.cpp
+++ b/Engine/Source/Developer/CollisionAnalyzer/Private/SCAQueryDetails.cpp
@@ -5,32 +5,32 @@
 #define LOCTEXT_NAMESPACE "SCAQueryDetails"
 
 /** Util to give written explanation for why we missed something */
-FString GetReasonForMiss(const UPrimitiveComponent* MissedComp, const FCAQuery* Query)
+FText GetReasonForMiss(const UPrimitiveComponent* MissedComp, const FCAQuery* Query)
 {
 	if(MissedComp != NULL && Query != NULL)
 	{
 		if(MissedComp->GetOwner() && !MissedComp->GetOwner()->GetActorEnableCollision())
 		{
-			return FString::Printf(TEXT("Owning Actor '%s' has all collision disabled (SetActorEnableCollision)"), *MissedComp->GetOwner()->GetName());
+			return FText::Format(LOCTEXT("MissReasonActorCollisionDisabledFmt", "Owning Actor '{0}' has all collision disabled (SetActorEnableCollision)"), FText::FromString(MissedComp->GetOwner()->GetName()));
 		}
 
 		if(!MissedComp->IsCollisionEnabled())
 		{
-			return FString::Printf(TEXT("Component '%s' has CollisionEnabled == NoCollision"), *MissedComp->GetName());
+			return FText::Format(LOCTEXT("MissReasonComponentCollisionDisabledFmt", "Component '{0}' has CollisionEnabled == NoCollision"), FText::FromString(MissedComp->GetName()));
 		}
 
 		if(MissedComp->GetCollisionResponseToChannel(Query->Channel) == ECR_Ignore)
 		{
-			return FString::Printf(TEXT("Component '%s' ignores this channel."), *MissedComp->GetName());
+			return FText::Format(LOCTEXT("MissReasonComponentIgnoresChannelFmt", "Component '{0}' ignores this channel."), FText::FromString(MissedComp->GetName()));
 		}
 
 		if(Query->ResponseParams.CollisionResponse.GetResponse(MissedComp->GetCollisionObjectType()) == ECR_Ignore)
 		{
-			return FString::Printf(TEXT("Query ignores Component '%s' movement channel."), *MissedComp->GetName());
+			return FText::Format(LOCTEXT("MissReasonQueryIgnoresComponentFmt", "Query ignores Component '{0}' movement channel."), FText::FromString(MissedComp->GetName()));
 		}
 	}
 
-	return TEXT("Unknown");
+	return LOCTEXT("MissReasonUnknown", "Unknown");
 }
 
 /** Implements a row widget for result list. */
@@ -57,11 +57,11 @@ public:
 	{
 		// Get info to apply to all columns (color and tooltip)
 		FSlateColor ResultColor = FSlateColor::UseForeground();
-		FString TooltipString;
+		FText TooltipText = FText::GetEmpty();
 		if(Info->bMiss)
 		{
 			ResultColor = FLinearColor(0.4f,0.4f,0.65f);
-			TooltipString = LOCTEXT("MissPrefix", "Miss: ").ToString() + GetReasonForMiss(Info->Result.Component.Get(), OwnerDetailsPtr.Pin()->GetCurrentQuery());
+			TooltipText = FText::Format(LOCTEXT("MissToolTipFmt", "Miss: {0}"), GetReasonForMiss(Info->Result.Component.Get(), OwnerDetailsPtr.Pin()->GetCurrentQuery()));
 		}
 		else if(Info->Result.bBlockingHit && Info->Result.bStartPenetrating)
 		{
@@ -72,51 +72,55 @@ public:
 		// Generate widget for column
 		if (ColumnName == TEXT("Time"))
 		{
+			static const FNumberFormattingOptions TimeNumberFormat = FNumberFormattingOptions()
+				.SetMinimumFractionalDigits(3)
+				.SetMaximumFractionalDigits(3);
+
 			return	SNew(STextBlock)
 					.ColorAndOpacity( ResultColor )
-					.ToolTipText( TooltipString )
-					.Text( FString::Printf(TEXT("%.3f"), Info->Result.Time) );
+					.ToolTipText( TooltipText )
+					.Text( FText::AsNumber(Info->Result.Time, &TimeNumberFormat) );
 		}
 		else if (ColumnName == TEXT("Type"))
 		{
-			FString TypeString;
+			FText TypeText = FText::GetEmpty();
 			if(Info->bMiss)
 			{
-				TypeString = TEXT("Miss");
+				TypeText = LOCTEXT("MissLabel", "Miss");
 			}
 			else if(Info->Result.bBlockingHit)
 			{
-				TypeString = TEXT("Block");
+				TypeText = LOCTEXT("BlockLabel", "Block");
 			}
 			else
 			{
-				TypeString = TEXT("Touch");
+				TypeText = LOCTEXT("TouchLabel", "Touch");
 			}
 
 			return	SNew(STextBlock)
 					.ColorAndOpacity( ResultColor )
-					.ToolTipText( TooltipString )
-					.Text( TypeString );
+					.ToolTipText( TooltipText )
+					.Text( TypeText );
 		}
 		else if (ColumnName == TEXT("Component"))
 		{
-			FString LongName = TEXT("Invalid");
+			FText LongName = LOCTEXT("InvalidLabel", "Invalid");
 			if(Info->Result.Component.IsValid())
 			{
-				LongName = Info->Result.Component.Get()->GetReadableName();
+				LongName = FText::FromString(Info->Result.Component.Get()->GetReadableName());
 			}
 
 			return	SNew(STextBlock)
 					.ColorAndOpacity( ResultColor )
-					.ToolTipText( TooltipString )
+					.ToolTipText( TooltipText )
 					.Text( LongName );
 		}
 		else if (ColumnName == TEXT("Normal"))
 		{
 			return	SNew(STextBlock)
 					.ColorAndOpacity( ResultColor )
-					.ToolTipText( TooltipString )
-					.Text( FString::Printf(TEXT("%s"), *Info->Result.Normal.ToString()) );
+					.ToolTipText( TooltipText )
+					.Text( FText::FromString(Info->Result.Normal.ToString()) );
 		}
 
 		return SNullWidget::NullWidget;
diff --git a/Engine/Source/Developer/CollisionAnalyzer/Private/SCAQueryTableRow.cpp b/Engine/Source/Developer/CollisionAnalyzer/Private/SCAQueryTableRow.cpp
--- a/Engine/Source/Developer/CollisionAnalyzer/Private/SCAQueryTableRow.cpp
+++ b/Engine/Source/Developer/CollisionAnalyzer/Private/SCAQueryTableRow.cpp
@@ -31,25 +31,25 @@ TSharedRef<SWidget> SCAQueryTableRow::GenerateWidgetForColumn(const FName& Colum
 		{
 			return	SNew(STextBlock)
 					.Font(FEditorStyle::GetFontStyle("BoldFont"))
-					.Text(FString::FromInt(Item->FrameNum));
+					.Text(FText::AsNumber(Item->FrameNum));
 		}
 		else if (ColumnName == TEXT("Tag") && OwnerAnalyzerWidget->GroupBy == EQueryGroupMode::ByTag)
 		{
 			return	SNew(STextBlock)
 					.Font(FEditorStyle::GetFontStyle("BoldFont"))
-					.Text(Item->GroupName.ToString());
+					.Text(FText::FromName(Item->GroupName));
 		}
 		else if (ColumnName == TEXT("Owner") && OwnerAnalyzerWidget->GroupBy == EQueryGroupMode::ByOwnerTag)
 		{
 			return	SNew(STextBlock)
 					.Font(FEditorStyle::GetFontStyle("BoldFont"))
-					.Text(Item->GroupName.ToString());
+					.Text(FText::FromName(Item->GroupName));
 		}
 		else if (ColumnName == TEXT("Time"))
 		{
 			return	SNew(STextBlock)
 					.Font(FEditorStyle::GetFontStyle("BoldFont"))
-					.Text(this, &SCAQueryTableRow::GetTotalTimeString);
+					.Text(this, &SCAQueryTableRow::GetTotalTimeText);
 		}
 	}
 	// ITEM
@@ -77,22 +77,22 @@ TSharedRef<SWidget> SCAQueryTableRow::GenerateWidgetForColumn(const FName& Colum
 		else if (ColumnName == TEXT("Type"))
 		{
 			return	SNew(STextBlock)
-					.Text(SCollisionAnalyzer::QueryTypeToString(Query.Type));
+					.Text(FText::FromString(SCollisionAnalyzer::QueryTypeToString(Query.Type)));
 		}
 		else if (ColumnName == TEXT("Shape"))
 		{
 			return	SNew(STextBlock)
-					.Text(SCollisionAnalyzer::QueryShapeToString(Query.Shape));
+					.Text(FText::FromString(SCollisionAnalyzer::QueryShapeToString(Query.Shape)));
 		}
 		else if (ColumnName == TEXT("Tag"))
 		{
 			return	SNew(STextBlock)
-					.Text(Query.Params.TraceTag.ToString());
+					.Text(FText::FromName(Query.Params.TraceTag));
 		}
 		else if (ColumnName == TEXT("Owner"))
 		{
 			return	SNew(STextBlock)
-					.Text(Query.Params.OwnerTag.ToString());
+					.Text(FText::FromName(Query.Params.OwnerTag));
 		}
 		else if (ColumnName == TEXT("NumBlock"))
 		{
@@ -100,28 +100,36 @@ TSharedRef<SWidget> SCAQueryTableRow::GenerateWidgetForColumn(const FName& Colum
 			bool bStartPenetrating = (FirstHit != NULL) && FirstHit->bStartPenetrating;
 			// Draw number in red if we start penetrating
 			return	SNew(STextBlock)
-					.Text(FString::FromInt(FHitResult::GetNumBlockingHits(Query.Results)))
+					.Text(FText::AsNumber(FHitResult::GetNumBlockingHits(Query.Results)))
 					.ColorAndOpacity(bStartPenetrating ? FLinearColor(1.f,0.25f,0.25f) : FSlateColor::UseForeground() );
 		}
 		else if (ColumnName == TEXT("NumTouch"))
 		{
 			return	SNew(STextBlock)
-					.Text(FString::FromInt(FHitResult::GetNumOverlapHits(Query.Results)));
+					.Text(FText::AsNumber(FHitResult::GetNumOverlapHits(Query.Results)));
 		}
 		else if (ColumnName == TEXT("Time"))
 		{
+			static const FNumberFormattingOptions TimeFormatOptions = FNumberFormattingOptions()
+				.SetMinimumFractionalDigits(3)
+				.SetMaximumFractionalDigits(3);
+
 			return	SNew(STextBlock)
-					.Text(FString::Printf(TEXT("%.3f"), Query.CPUTime));
+					.Text(FText::AsNumber(Query.CPUTime, &TimeFormatOptions));
 		}
 	}
 
 	return SNullWidget::NullWidget;
 }
 
-FString SCAQueryTableRow::GetTotalTimeString() const
+FText SCAQueryTableRow::GetTotalTimeText() const
 {
+	static const FNumberFormattingOptions TimeFormatOptions = FNumberFormattingOptions()
+		.SetMinimumFractionalDigits(3)
+		.SetMaximumFractionalDigits(3);
+
 	check(Item->bIsGroup)
-	return FString::Printf(TEXT("%.3f"), Item->TotalCPUTime);
+	return FText::AsNumber(Item->TotalCPUTime, &TimeFormatOptions);
 }
 
 #undef LOCTEXT_NAMESPACE
\ No newline at end of file
diff --git a/Engine/Source/Developer/CollisionAnalyzer/Private/SCAQueryTableRow.h b/Engine/Source/Developer/CollisionAnalyzer/Private/SCAQueryTableRow.h
--- a/Engine/Source/Developer/CollisionAnalyzer/Private/SCAQueryTableRow.h
+++ b/Engine/Source/Developer/CollisionAnalyzer/Private/SCAQueryTableRow.h
@@ -21,7 +21,7 @@ public:
 	virtual TSharedRef<SWidget> GenerateWidgetForColumn(const FName& ColumnName) override;
 
 	// delegate
-	FString GetTotalTimeString() const;
+	FText GetTotalTimeText() const;
 
 private:
 
diff --git a/Engine/Source/Developer/CollisionAnalyzer/Private/SCollisionAnalyzer.cpp b/Engine/Source/Developer/CollisionAnalyzer/Private/SCollisionAnalyzer.cpp
--- a/Engine/Source/Developer/CollisionAnalyzer/Private/SCollisionAnalyzer.cpp
+++ b/Engine/Source/Developer/CollisionAnalyzer/Private/SCollisionAnalyzer.cpp
@@ -660,9 +660,12 @@ const FSlateBrush* SCollisionAnalyzer::GetRecordButtonBrush() const
 	}
 }
 
-FString SCollisionAnalyzer::GetStatusText() const
+FText SCollisionAnalyzer::GetStatusText() const
 {
-	return FString::Printf(TEXT("Total: %d queries over %d frames. Shown: %d queries"), Analyzer->Queries.Num(), Analyzer->GetNumFramesOfRecording(), TotalNumQueries);
+	return FText::Format(
+		LOCTEXT("CollisionAnalyzerStatusTextFmt", "Total: {0} queries over {1} frames. Shown: {2} queries"), 
+		FText::AsNumber(Analyzer->Queries.Num()), FText::AsNumber(Analyzer->GetNumFramesOfRecording()), FText::AsNumber(TotalNumQueries)
+		);
 }
 
 ECheckBoxState SCollisionAnalyzer::GetDrawRecentState() const
diff --git a/Engine/Source/Developer/CollisionAnalyzer/Private/SCollisionAnalyzer.h b/Engine/Source/Developer/CollisionAnalyzer/Private/SCollisionAnalyzer.h
--- a/Engine/Source/Developer/CollisionAnalyzer/Private/SCollisionAnalyzer.h
+++ b/Engine/Source/Developer/CollisionAnalyzer/Private/SCollisionAnalyzer.h
@@ -94,7 +94,7 @@ public:
 
 	// Get delegates
 	const FSlateBrush* GetRecordButtonBrush() const;
-	FString GetStatusText() const;
+	FText GetStatusText() const;
 	ECheckBoxState GetDrawRecentState() const;
 	ECheckBoxState GetGroupByFrameState() const;
 	ECheckBoxState GetGroupByTagState() const;
diff --git a/Engine/Source/Developer/DeviceManager/Private/Widgets/Browser/SDeviceBrowserDeviceAdder.cpp b/Engine/Source/Developer/DeviceManager/Private/Widgets/Browser/SDeviceBrowserDeviceAdder.cpp
--- a/Engine/Source/Developer/DeviceManager/Private/Widgets/Browser/SDeviceBrowserDeviceAdder.cpp
+++ b/Engine/Source/Developer/DeviceManager/Private/Widgets/Browser/SDeviceBrowserDeviceAdder.cpp
@@ -290,11 +290,11 @@ void SDeviceBrowserDeviceAdder::HandleDeviceNameTextBoxTextChanged (const FStrin
 }
 
 
-FString SDeviceBrowserDeviceAdder::HandlePlatformComboBoxContentText( ) const
+FText SDeviceBrowserDeviceAdder::HandlePlatformComboBoxContentText( ) const
 {
 	TSharedPtr<FString> SelectedPlatform = PlatformComboBox->GetSelectedItem();
 
-	return SelectedPlatform.IsValid() ? *SelectedPlatform : FString(TEXT("Select a Platform"));
+	return SelectedPlatform.IsValid() ? FText::FromString(*SelectedPlatform) : LOCTEXT("SelectAPlatform", "Select a Platform");
 }
 
 
@@ -323,7 +323,7 @@ TSharedRef<SWidget> SDeviceBrowserDeviceAdder::HandlePlatformComboBoxGenerateWid
 			.VAlign(VAlign_Center)
 			[
 				SNew(STextBlock)
-					.Text(*Item)
+					.Text(FText::FromString(*Item))
 			];
 }
 
diff --git a/Engine/Source/Developer/DeviceManager/Private/Widgets/Browser/SDeviceBrowserDeviceAdder.h b/Engine/Source/Developer/DeviceManager/Private/Widgets/Browser/SDeviceBrowserDeviceAdder.h
--- a/Engine/Source/Developer/DeviceManager/Private/Widgets/Browser/SDeviceBrowserDeviceAdder.h
+++ b/Engine/Source/Developer/DeviceManager/Private/Widgets/Browser/SDeviceBrowserDeviceAdder.h
@@ -51,7 +51,7 @@ private:
 	void HandleDeviceNameTextBoxTextChanged ( const FString& Text );
 
 	// Callback for getting the name of the selected platform.
-	FString HandlePlatformComboBoxContentText( ) const;
+	FText HandlePlatformComboBoxContentText( ) const;
 
 	// Callback for generating widgets for the platforms combo box.
 	TSharedRef<SWidget> HandlePlatformComboBoxGenerateWidget( TSharedPtr<FString> StringItem );
diff --git a/Engine/Source/Developer/DeviceManager/Private/Widgets/Browser/SDeviceBrowserDeviceListRow.h b/Engine/Source/Developer/DeviceManager/Private/Widgets/Browser/SDeviceBrowserDeviceListRow.h
--- a/Engine/Source/Developer/DeviceManager/Private/Widgets/Browser/SDeviceBrowserDeviceListRow.h
+++ b/Engine/Source/Developer/DeviceManager/Private/Widgets/Browser/SDeviceBrowserDeviceListRow.h
@@ -92,7 +92,7 @@ public:
 				[
 					SNew(STextBlock)
 						.ColorAndOpacity(this, &SDeviceBrowserDeviceListRow::HandleTextColorAndOpacity)
-						.Text(DeviceService->GetDevicePlatformDisplayName())
+						.Text(FText::FromString(DeviceService->GetDevicePlatformDisplayName()))
 				];
 		}
 		else if (ColumnName == TEXT("Share"))
@@ -105,7 +105,7 @@ public:
 						.IsChecked(this, &SDeviceBrowserDeviceListRow::HandleShareCheckBoxIsChecked)
 						.IsEnabled(this, &SDeviceBrowserDeviceListRow::HandleShareCheckBoxIsEnabled)
 						.OnCheckStateChanged(this, &SDeviceBrowserDeviceListRow::HandleShareCheckBoxStateChanged)
-						.ToolTipText(LOCTEXT("ShareCheckBoxToolTip", "Check this box to share this device with other users on the network").ToString())
+						.ToolTipText(LOCTEXT("ShareCheckBoxToolTip", "Check this box to share this device with other users on the network"))
 				];
 		}
 		else if (ColumnName == TEXT("Status"))
diff --git a/Engine/Source/Developer/DeviceManager/Private/Widgets/Browser/SDeviceBrowserFilterBar.cpp b/Engine/Source/Developer/DeviceManager/Private/Widgets/Browser/SDeviceBrowserFilterBar.cpp
--- a/Engine/Source/Developer/DeviceManager/Private/Widgets/Browser/SDeviceBrowserFilterBar.cpp
+++ b/Engine/Source/Developer/DeviceManager/Private/Widgets/Browser/SDeviceBrowserFilterBar.cpp
@@ -106,9 +106,9 @@ ECheckBoxState SDeviceBrowserFilterBar::HandlePlatformListRowIsChecked(TSharedPt
 }
 
 
-FString SDeviceBrowserFilterBar::HandlePlatformListRowText(TSharedPtr<FDeviceBrowserFilterEntry> PlatformEntry) const
+FText SDeviceBrowserFilterBar::HandlePlatformListRowText(TSharedPtr<FDeviceBrowserFilterEntry> PlatformEntry) const
 {
-	return FString::Printf(TEXT("%s (%i)"), *(PlatformEntry->PlatformName), Filter->GetServiceCountPerPlatform(PlatformEntry->PlatformName));
+	return FText::Format(LOCTEXT("PlatformListRowFmt", "{0} ({1})"), FText::FromString(PlatformEntry->PlatformName), FText::AsNumber(Filter->GetServiceCountPerPlatform(PlatformEntry->PlatformName)));
 }
 
 BEGIN_SLATE_FUNCTION_BUILD_OPTIMIZATION
diff --git a/Engine/Source/Developer/DeviceManager/Private/Widgets/Browser/SDeviceBrowserFilterBar.h b/Engine/Source/Developer/DeviceManager/Private/Widgets/Browser/SDeviceBrowserFilterBar.h
--- a/Engine/Source/Developer/DeviceManager/Private/Widgets/Browser/SDeviceBrowserFilterBar.h
+++ b/Engine/Source/Developer/DeviceManager/Private/Widgets/Browser/SDeviceBrowserFilterBar.h
@@ -53,7 +53,7 @@ private:
 	ECheckBoxState HandlePlatformListRowIsChecked(TSharedPtr<FDeviceBrowserFilterEntry> PlatformEntry) const;
 
 	// Callback for getting the text for a row in the platform filter drop-down.
-	FString HandlePlatformListRowText(TSharedPtr<FDeviceBrowserFilterEntry> PlatformEntry) const;
+	FText HandlePlatformListRowText(TSharedPtr<FDeviceBrowserFilterEntry> PlatformEntry) const;
 
 	// Generates a row widget for the platform filter list.
 	TSharedRef<ITableRow> HandlePlatformListViewGenerateRow(TSharedPtr<FDeviceBrowserFilterEntry> PlatformEntry, const TSharedRef<STableViewBase>& OwnerTable);
diff --git a/Engine/Source/Developer/DeviceManager/Private/Widgets/Processes/SDeviceProcessesProcessListRow.h b/Engine/Source/Developer/DeviceManager/Private/Widgets/Processes/SDeviceProcessesProcessListRow.h
--- a/Engine/Source/Developer/DeviceManager/Private/Widgets/Processes/SDeviceProcessesProcessListRow.h
+++ b/Engine/Source/Developer/DeviceManager/Private/Widgets/Processes/SDeviceProcessesProcessListRow.h
@@ -51,7 +51,7 @@ public:
 				.VAlign(VAlign_Center)
 				[
 					SNew(STextBlock)
-						.Text(Node->GetProcessInfo().Name)
+						.Text(FText::FromString(Node->GetProcessInfo().Name))
 				];
 
 			if (OwnerTablePtr.Pin()->GetTableViewMode() == ETableViewMode::Tree)
@@ -84,7 +84,7 @@ public:
 				.VAlign(VAlign_Center)
 				[
 					SNew(STextBlock)
-						.Text(FString::Printf(TEXT("%d"), Node->GetProcessInfo().ParentId))
+						.Text(FText::AsNumber(Node->GetProcessInfo().ParentId))
 				];
 		}
 		else if (ColumnName == TEXT("PID"))
@@ -94,7 +94,7 @@ public:
 				.VAlign(VAlign_Center)
 				[
 					SNew(STextBlock)
-						.Text(FString::Printf(TEXT("%d"), Node->GetProcessInfo().Id))
+						.Text(FText::AsNumber(Node->GetProcessInfo().Id))
 				];
 		}
 		else if (ColumnName == TEXT("Threads"))
@@ -104,7 +104,7 @@ public:
 				.VAlign(VAlign_Center)
 				[
 					SNew(STextBlock)
-						.Text(FString::Printf(TEXT("%d"), Node->GetProcessInfo().Threads.Num()))
+						.Text(FText::AsNumber(Node->GetProcessInfo().Threads.Num()))
 				];
 		}
 		else if (ColumnName == TEXT("User"))
@@ -114,7 +114,7 @@ public:
 				.VAlign(VAlign_Center)
 				[
 					SNew(STextBlock)
-						.Text(Node->GetProcessInfo().UserName)
+						.Text(FText::FromString(Node->GetProcessInfo().UserName))
 				];
 		}
 
diff --git a/Engine/Source/Developer/GammaUI/Private/GammaUIPanel.cpp b/Engine/Source/Developer/GammaUI/Private/GammaUIPanel.cpp
--- a/Engine/Source/Developer/GammaUI/Private/GammaUIPanel.cpp
+++ b/Engine/Source/Developer/GammaUI/Private/GammaUIPanel.cpp
@@ -16,7 +16,7 @@ void SGammaUIPanel::Construct(const SGammaUIPanel::FArguments& InArgs)
 			.Padding( 0.0f, 4.0f, 0.0f, 4.0f )
 		[
 			SNew( STextBlock )
-				.Text( NSLOCTEXT( "GammaUI", "GammaUILabel", "Gamma" ).ToString() )
+				.Text( NSLOCTEXT( "GammaUI", "GammaUILabel", "Gamma" ) )
 		]
 
 		+ SVerticalBox::Slot()
diff --git a/Engine/Source/Developer/Merge/Private/SMergeAssetPickerView.cpp b/Engine/Source/Developer/Merge/Private/SMergeAssetPickerView.cpp
--- a/Engine/Source/Developer/Merge/Private/SMergeAssetPickerView.cpp
+++ b/Engine/Source/Developer/Merge/Private/SMergeAssetPickerView.cpp
@@ -133,7 +133,7 @@ TSharedRef<SWidget> SMergeAssetPickerPanel::MakeAssetPicker()
 		.ButtonStyle(FEditorStyle::Get(), "NoBorder")
 		.OnClicked(FOnClicked::CreateSP(this, &SMergeAssetPickerPanel::OnBrowseToPickedAsset))
 		.ContentPadding(0)
-		.ToolTipText(NSLOCTEXT("GraphEditor", "ObjectGraphPin_Browse", "Show the selected asset in the content browser.").ToString())
+		.ToolTipText(NSLOCTEXT("GraphEditor", "ObjectGraphPin_Browse", "Show the selected asset in the content browser."))
 	[
 		SNew(SImage)
 			.Image(FEditorStyle::GetBrush(TEXT("PropertyWindow.Button_Browse")))
@@ -169,7 +169,7 @@ TSharedRef<SWidget> SMergeAssetPickerPanel::MakeAssetPicker()
 			.ButtonStyle(FEditorStyle::Get(), "NoBorder")
 			.OnClicked(FOnClicked::CreateSP(this, &SMergeAssetPickerPanel::OnUseSelectedAssetClick))
 			.ContentPadding(1.f)
-			.ToolTipText(NSLOCTEXT("GraphEditor", "ObjectGraphPin_Use", "Use content browser selection.").ToString())
+			.ToolTipText(NSLOCTEXT("GraphEditor", "ObjectGraphPin_Use", "Use content browser selection."))
 		[
 			SNew(SImage)
 				.Image(FEditorStyle::GetBrush(TEXT("PropertyWindow.Button_Use")))
diff --git a/Engine/Source/Developer/Merge/Private/SMergeGraphView.cpp b/Engine/Source/Developer/Merge/Private/SMergeGraphView.cpp
--- a/Engine/Source/Developer/Merge/Private/SMergeGraphView.cpp
+++ b/Engine/Source/Developer/Merge/Private/SMergeGraphView.cpp
@@ -198,11 +198,11 @@ static TArray< FMergeGraphEntry > GenerateDiffListItems(const FBlueprintRevPair&
 
 						if( ConflictingDifference )
 						{
-							Label = FText::Format( NSLOCTEXT("SMergeGraphView", "ConflictIdentifier", "CONFLICT: {0} conflicts with {1}" ), FText::FromString((*ConflictingDifference)->DisplayString), FText::FromString(Difference.DisplayString) );
+							Label = FText::Format( NSLOCTEXT("SMergeGraphView", "ConflictIdentifier", "CONFLICT: {0} conflicts with {1}" ), (*ConflictingDifference)->DisplayString, Difference.DisplayString );
 						}
 						else
 						{
-							Label = FText::FromString( Difference.DisplayString );
+							Label = Difference.DisplayString;
 						}
 
 						FMergeGraphRowEntry NewEntry = {
@@ -231,7 +231,7 @@ static TArray< FMergeGraphEntry > GenerateDiffListItems(const FBlueprintRevPair&
 						if (!ConflictingDifference)
 						{
 							FMergeGraphRowEntry NewEntry = {
-								FText::FromString(Difference.DisplayString)
+								Difference.DisplayString
 								, Difference.OwningGraph
 								, Difference.Node2 /*UEdGraphNode* LocalNode*/
 								, Difference.Node1 /*UEdGraphNode* BaseNode*/
diff --git a/Engine/Source/Developer/MessageLog/Private/UserInterface/SMessageLogListing.cpp b/Engine/Source/Developer/MessageLog/Private/UserInterface/SMessageLogListing.cpp
--- a/Engine/Source/Developer/MessageLog/Private/UserInterface/SMessageLogListing.cpp
+++ b/Engine/Source/Developer/MessageLog/Private/UserInterface/SMessageLogListing.cpp
@@ -342,7 +342,7 @@ TSharedRef<ITableRow> SMessageLogListing::MakeShowWidget(TSharedRef<FMessageFilt
 				+SHorizontalBox::Slot().AutoWidth()
 				[
 					SNew(STextBlock)
-					.Text(Selection->GetName().ToString())
+					.Text(Selection->GetName())
 				]
 			]
 		];
diff --git a/Engine/Source/Developer/OutputLog/Private/SOutputLog.cpp b/Engine/Source/Developer/OutputLog/Private/SOutputLog.cpp
--- a/Engine/Source/Developer/OutputLog/Private/SOutputLog.cpp
+++ b/Engine/Source/Developer/OutputLog/Private/SOutputLog.cpp
@@ -229,7 +229,7 @@ TSharedRef<ITableRow> SConsoleInputBox::MakeSuggestionListItemWidget(TSharedPtr<
 			.WidthOverride(300)			// to enforce some minimum width, ideally we define the minimum, not a fixed width
 			[
 				SNew(STextBlock)
-				.Text(Combined)
+				.Text(FText::FromString(Combined))
 				.TextStyle( FEditorStyle::Get(), "Log.Normal")
 				.HighlightText(HighlightText)
 			]
diff --git a/Engine/Source/Developer/Profiler/Private/Widgets/SDataGraph.cpp b/Engine/Source/Developer/Profiler/Private/Widgets/SDataGraph.cpp
--- a/Engine/Source/Developer/Profiler/Private/Widgets/SDataGraph.cpp
+++ b/Engine/Source/Developer/Profiler/Private/Widgets/SDataGraph.cpp
@@ -120,9 +120,9 @@ protected:
 		return FReply::Handled();
 	}
 
-	FString SummaryInformation_GetSummary() const
+	FText SummaryInformation_GetSummary() const
 	{
-		FString SummaryText = LOCTEXT("DataGraphSummary_Warning", "Not implemented yet").ToString();
+		FText SummaryText = LOCTEXT("DataGraphSummary_Warning", "Not implemented yet");
 		const bool bCanDisplayData = GraphDescription.CombinedGraphDataSource->CanBeDisplayedAsIndexBased() && ParentWidget->GetViewMode() == EDataGraphViewModes::Index;
 
 		const uint32 FrameIndex = OnGetMouseFrameIndex.IsBound() ? (uint32)OnGetMouseFrameIndex.Execute() : 0;
@@ -134,21 +134,24 @@ protected:
 			const EProfilerSampleTypes::Type UnitType = (*GraphDataSource)->GetSampleType();
 			const FProfilerAggregatedStat& Aggregated = *(*GraphDataSource)->GetAggregatedStat();
 
-			SummaryText = FString::Printf( TEXT("%4.2f - "), SampleValue );
-			SummaryText+= Aggregated.ToString();
+			static const FNumberFormattingOptions SampleValueFormattingOptions = FNumberFormattingOptions()
+				.SetMinimumFractionalDigits(2)
+				.SetMaximumFractionalDigits(2);
+
+			SummaryText = FText::Format(LOCTEXT("DataGraphSummary_SummaryFmt", "{0} - {1}"), FText::AsNumber(SampleValue, &SampleValueFormattingOptions), FText::FromString(Aggregated.ToString()));
 		}
 
 		return SummaryText;
 	}
 
-	FString SummaryInformation_GetGroupName() const
+	FText SummaryInformation_GetGroupName() const
 	{
-		return FString::Printf( TEXT("(%s)"), *GraphDescription.CombinedGraphDataSource->GetGroupName() );
+		return FText::Format(LOCTEXT("DataGraphSummary_GroupNameFmt", "({0})"), FText::FromString(GraphDescription.CombinedGraphDataSource->GetGroupName()));
 	}
 
-	FString SummaryInformation_GetStatName() const
+	FText SummaryInformation_GetStatName() const
 	{
-		return FProfilerHelper::ShortenName( GraphDescription.CombinedGraphDataSource->GetStatName(), 32 );
+		return FText::FromString(FProfilerHelper::ShortenName( GraphDescription.CombinedGraphDataSource->GetStatName(), 32 ));
 	}
 
 private:
diff --git a/Engine/Source/Developer/Profiler/Private/Widgets/SEventGraph.cpp b/Engine/Source/Developer/Profiler/Private/Widgets/SEventGraph.cpp
--- a/Engine/Source/Developer/Profiler/Private/Widgets/SEventGraph.cpp
+++ b/Engine/Source/Developer/Profiler/Private/Widgets/SEventGraph.cpp
@@ -388,7 +388,7 @@ protected:
 			.Padding( FMargin( 2.0f, 0.0f ) )
 			[
 				SNew( STextBlock )
-				.Text( EventPtr->_StatName.ToString() )
+				.Text( FText::FromName(EventPtr->_StatName) )
 				.TextStyle( FEditorStyle::Get(), TEXT("Profiler.Tooltip") )
 				.ColorAndOpacity( this, &SEventGraphTableCell::GetColorAndOpacity )
 				.ShadowColorAndOpacity( this, &SEventGraphTableCell::GetShadowColorAndOpacity )
@@ -438,7 +438,7 @@ protected:
 			.Padding( FMargin( 2.0f, 0.0f ) )
 			[
 				SNew( STextBlock )
-				.Text( FormattedValue )
+				.Text( FText::FromString(FormattedValue) )
 				.TextStyle( FEditorStyle::Get(), TEXT("Profiler.Tooltip") )
 				.ColorAndOpacity( this, &SEventGraphTableCell::GetColorAndOpacity )
 				.ShadowColorAndOpacity( this, &SEventGraphTableCell::GetShadowColorAndOpacity )
@@ -950,7 +950,7 @@ TSharedRef<SWidget> SEventGraph::GetToggleButtonForEventGraphType( const EEventG
 	.VAlign( VAlign_Center )
 	[
 		SNew( STextBlock )
-		.Text( EEventGraphTypes::ToName( EventGraphType ) )
+		.Text( FText::FromString(EEventGraphTypes::ToName( EventGraphType )) )
 		.TextStyle( FEditorStyle::Get(), TEXT("Profiler.Caption") )
 	];
 
@@ -962,7 +962,7 @@ TSharedRef<SWidget> SEventGraph::GetToggleButtonForEventGraphType( const EEventG
 	.Padding( 2.0f )
 	.OnCheckStateChanged( this, &SEventGraph::EventGraphType_OnCheckStateChanged, EventGraphType )
 	.IsChecked( this, &SEventGraph::EventGraphType_IsChecked, EventGraphType )
-	.ToolTipText( EEventGraphTypes::ToDescription( EventGraphType ) )
+	.ToolTipText( FText::FromString(EEventGraphTypes::ToDescription( EventGraphType )) )
 	[
 		ButtonContent
 	];
@@ -2869,7 +2869,7 @@ TSharedRef<SHorizontalBox> SEventGraph::GetContentForEvent( FEventGraphSamplePtr
 	.VAlign(VAlign_Center)
 	[
 		SNew(STextBlock)
-		.Text( GetEventDescription(EventPtr,Pct,bSimple) )
+		.Text( FText::FromString(GetEventDescription(EventPtr,Pct,bSimple)) )
 		.TextStyle( FEditorStyle::Get(), bSimple ? TEXT("Profiler.Tooltip") : TEXT("Profiler.EventGraph.DarkText") )
 	];
 
diff --git a/Engine/Source/Developer/Profiler/Private/Widgets/SEventGraphTooltip.cpp b/Engine/Source/Developer/Profiler/Private/Widgets/SEventGraphTooltip.cpp
--- a/Engine/Source/Developer/Profiler/Private/Widgets/SEventGraphTooltip.cpp
+++ b/Engine/Source/Developer/Profiler/Private/Widgets/SEventGraphTooltip.cpp
@@ -127,7 +127,7 @@ TSharedPtr<SToolTip> SEventGraphTooltip::GetTableCellTooltip( const FEventGraphS
 							.Padding( 2.0f )
 							[
 								SNew(STextBlock)
-								.Text( EventSample->GetShortEventName() )
+								.Text( FText::FromString(EventSample->GetShortEventName()) )
 								.TextStyle( FEditorStyle::Get(), TEXT("Profiler.Tooltip") )
 							]
 
@@ -162,7 +162,7 @@ TSharedPtr<SToolTip> SEventGraphTooltip::GetTableCellTooltip( const FEventGraphS
 							.Padding( 2.0f )
 							[
 								SNew(STextBlock)
-								.Text( EventSample->GetFormattedValue(EEventPropertyIndex::InclusiveTimeMS) )
+								.Text( FText::FromString(EventSample->GetFormattedValue(EEventPropertyIndex::InclusiveTimeMS)) )
 								.TextStyle( FEditorStyle::Get(), TEXT("Profiler.Tooltip") )
 							]
 						// ([_InclusiveTimePct] % of the caller)
@@ -186,7 +186,7 @@ TSharedPtr<SToolTip> SEventGraphTooltip::GetTableCellTooltip( const FEventGraphS
 							.Padding( 2.0f )
 							[
 								SNew(STextBlock)
-								.Text( EventSample->GetFormattedValue(EEventPropertyIndex::AvgInclusiveTimePerCallMS) )
+								.Text( FText::FromString(EventSample->GetFormattedValue(EEventPropertyIndex::AvgInclusiveTimePerCallMS)) )
 								.TextStyle( FEditorStyle::Get(), TEXT("Profiler.Tooltip") )
 							]
 
@@ -206,7 +206,7 @@ TSharedPtr<SToolTip> SEventGraphTooltip::GetTableCellTooltip( const FEventGraphS
 							.Padding( 2.0f )
 							[
 								SNew(STextBlock)
-								.Text( EventSample->GetFormattedValue(EEventPropertyIndex::ThreadPct) )
+								.Text( FText::FromString(EventSample->GetFormattedValue(EEventPropertyIndex::ThreadPct)) )
 								.TextStyle( FEditorStyle::Get(), TEXT("Profiler.Tooltip") )
 							]
 
@@ -222,7 +222,7 @@ TSharedPtr<SToolTip> SEventGraphTooltip::GetTableCellTooltip( const FEventGraphS
 							.Padding( 2.0f )
 							[
 								SNew(STextBlock)
-								.Text( EventSample->GetFormattedValue(EEventPropertyIndex::FramePct) )
+								.Text( FText::FromString(EventSample->GetFormattedValue(EEventPropertyIndex::FramePct)) )
 								.TextStyle( FEditorStyle::Get(), TEXT("Profiler.Tooltip") )
 							]
 
@@ -242,7 +242,7 @@ TSharedPtr<SToolTip> SEventGraphTooltip::GetTableCellTooltip( const FEventGraphS
 							.Padding( 2.0f )
 							[
 								SNew(STextBlock)
-								.Text( EventSample->GetFormattedValue(EEventPropertyIndex::MinInclusiveTimeMS) )
+								.Text( FText::FromString(EventSample->GetFormattedValue(EEventPropertyIndex::MinInclusiveTimeMS)) )
 								.TextStyle( FEditorStyle::Get(), TEXT("Profiler.Tooltip") )
 							]
 						// [_MaxInclusiveTimeMS] ms
@@ -250,7 +250,7 @@ TSharedPtr<SToolTip> SEventGraphTooltip::GetTableCellTooltip( const FEventGraphS
 							.Padding( 2.0f )
 							[
 								SNew(STextBlock)
-								.Text( EventSample->GetFormattedValue(EEventPropertyIndex::MaxInclusiveTimeMS) )
+								.Text( FText::FromString(EventSample->GetFormattedValue(EEventPropertyIndex::MaxInclusiveTimeMS)) )
 								.TextStyle( FEditorStyle::Get(), TEXT("Profiler.Tooltip") )
 							]
 						// [_AvgInclusiveTimeMS] ms
@@ -258,7 +258,7 @@ TSharedPtr<SToolTip> SEventGraphTooltip::GetTableCellTooltip( const FEventGraphS
 							.Padding( 2.0f )
 							[
 								SNew(STextBlock)
-								.Text( EventSample->GetFormattedValue(EEventPropertyIndex::AvgInclusiveTimeMS) )
+								.Text( FText::FromString(EventSample->GetFormattedValue(EEventPropertyIndex::AvgInclusiveTimeMS)) )
 								.TextStyle( FEditorStyle::Get(), TEXT("Profiler.Tooltip") )
 							]
 
@@ -278,7 +278,7 @@ TSharedPtr<SToolTip> SEventGraphTooltip::GetTableCellTooltip( const FEventGraphS
 							.Padding( 2.0f )
 							[
 								SNew(STextBlock)
-								.Text( EventSample->GetFormattedValue(EEventPropertyIndex::ExclusiveTimeMS) )
+								.Text( FText::FromString(EventSample->GetFormattedValue(EEventPropertyIndex::ExclusiveTimeMS)) )
 								.TextStyle( FEditorStyle::Get(), TEXT("Profiler.Tooltip") )
 							]
 						// ([_ExclusiveTimePct] % of this call's inc time)
@@ -306,7 +306,7 @@ TSharedPtr<SToolTip> SEventGraphTooltip::GetTableCellTooltip( const FEventGraphS
 							.Padding( 2.0f )
 							[
 								SNew(STextBlock)
-								.Text( EventSample->GetFormattedValue(EEventPropertyIndex::NumCallsPerFrame) )
+								.Text( FText::FromString(EventSample->GetFormattedValue(EEventPropertyIndex::NumCallsPerFrame)) )
 								.TextStyle( FEditorStyle::Get(), TEXT("Profiler.Tooltip") )
 							]
 
@@ -323,7 +323,7 @@ TSharedPtr<SToolTip> SEventGraphTooltip::GetTableCellTooltip( const FEventGraphS
 							.Padding( 2.0f )
 							[
 								SNew(STextBlock)
-								.Text( EventSample->GetFormattedValue(EEventPropertyIndex::AvgNumCallsPerFrame) )
+								.Text( FText::FromString(EventSample->GetFormattedValue(EEventPropertyIndex::AvgNumCallsPerFrame)) )
 								.TextStyle( FEditorStyle::Get(), TEXT("Profiler.Tooltip") )
 							]
 
@@ -343,7 +343,7 @@ TSharedPtr<SToolTip> SEventGraphTooltip::GetTableCellTooltip( const FEventGraphS
 							.Padding( 2.0f )
 							[
 								SNew(STextBlock)
-								.Text( EventSample->GetFormattedValue(EEventPropertyIndex::ThreadDurationMS) )
+								.Text( FText::FromString(EventSample->GetFormattedValue(EEventPropertyIndex::ThreadDurationMS)) )
 								.TextStyle( FEditorStyle::Get(), TEXT("Profiler.Tooltip") )
 							]
 
@@ -360,7 +360,7 @@ TSharedPtr<SToolTip> SEventGraphTooltip::GetTableCellTooltip( const FEventGraphS
 							.Padding( 2.0f )
 							[
 								SNew(STextBlock)
-								.Text( EventSample->GetFormattedValue(EEventPropertyIndex::FrameDurationMS) )
+								.Text( FText::FromString(EventSample->GetFormattedValue(EEventPropertyIndex::FrameDurationMS)) )
 								.TextStyle( FEditorStyle::Get(), TEXT("Profiler.Tooltip") )
 							]
 
@@ -378,7 +378,7 @@ TSharedPtr<SToolTip> SEventGraphTooltip::GetTableCellTooltip( const FEventGraphS
 							.Padding( 2.0f )
 							[
 								SNew(STextBlock)
-								.Text( EventSample->GetFormattedValue(EEventPropertyIndex::ThreadToFramePct) )
+								.Text( FText::FromString(EventSample->GetFormattedValue(EEventPropertyIndex::ThreadToFramePct)) )
 								.TextStyle( FEditorStyle::Get(), TEXT("Profiler.Tooltip") )
 							]
 					]
@@ -405,7 +405,7 @@ TSharedPtr<SToolTip> SEventGraphTooltip::GetTableCellTooltip( const FEventGraphS
 			.AutoWidth()
 			[
 				SNew(STextBlock)
-				.Text( ParentName )
+				.Text( FText::FromString(ParentName) )
 				.TextStyle( FEditorStyle::Get(), TEXT("Profiler.Caption") )
 			];
 
@@ -421,7 +421,7 @@ TSharedPtr<SToolTip> SEventGraphTooltip::GetTableCellTooltip( const FEventGraphS
 		.AutoWidth()
 		[
 			SNew(STextBlock)
-			.Text( EventSample->GetShortEventName() )
+			.Text( FText::FromString(EventSample->GetShortEventName()) )
 			.TextStyle( FEditorStyle::Get(), TEXT("Profiler.CaptionBold") )
 		];
 
@@ -471,7 +471,7 @@ TSharedPtr<SToolTip> SEventGraphTooltip::GetTableCellTooltip( const FEventGraphS
 			.AutoWidth()
 			[
 				SNew(STextBlock)
-				.Text( ChildrenNames )
+				.Text( FText::FromString(ChildrenNames) )
 				.TextStyle( FEditorStyle::Get(), TEXT("Profiler.Caption") )
 			];
 	}
diff --git a/Engine/Source/Developer/Profiler/Private/Widgets/SFiltersAndPresets.cpp b/Engine/Source/Developer/Profiler/Private/Widgets/SFiltersAndPresets.cpp
--- a/Engine/Source/Developer/Profiler/Private/Widgets/SFiltersAndPresets.cpp
+++ b/Engine/Source/Developer/Profiler/Private/Widgets/SFiltersAndPresets.cpp
@@ -119,7 +119,7 @@ protected:
 		[
 			SNew(STextBlock)
 			.TextStyle( FEditorStyle::Get(), TEXT("Profiler.Tooltip") )
-			.Text( InstanceName )
+			.Text( FText::FromString(InstanceName) )
 		];
 
 		AddSeparator( Grid, RowPos );
@@ -194,7 +194,7 @@ protected:
 			.VAlign(VAlign_Center)
 			[
 				SNew(STextBlock)
-				.Text( EProfilerSampleTypes::ToDescription(SampleType) )
+				.Text( FText::FromString(EProfilerSampleTypes::ToDescription(SampleType)) )
 				.TextStyle( FEditorStyle::Get(), TEXT("Profiler.Tooltip") )
 			]
 		];
@@ -413,7 +413,7 @@ public:
 			.Padding( FMargin( 2.0f, 0.0f ) )
 			[
 				SNew( STextBlock )
-				.Text( this, &SGroupAndStatTableRow::GetText/*EventPtr->_StatName.ToString()*/ )
+				.Text( this, &SGroupAndStatTableRow::GetText )
 				.HighlightText( InArgs._HighlightText )
 				.TextStyle( FEditorStyle::Get(), TEXT("Profiler.Tooltip") )
 				.ColorAndOpacity( this, &SGroupAndStatTableRow::GetColorAndOpacity )
@@ -486,9 +486,9 @@ protected:
 	/**
 	 * @return a text which describes this table row, refers to both groups and stats
 	 */
-	FString GetText() const
+	FText GetText() const
 	{
-		FString Text;
+		FText Text = FText::GetEmpty();
 
 		if( GroupOrStatNode->IsGroup() )
 		{
@@ -506,12 +506,19 @@ protected:
 				}
 			}
 
-			Text = FString::Printf( TEXT("%s (%i) (%i)"), *GroupOrStatNode->GetName().GetPlainNameString(), GroupOrStatNode->GetChildren().Num(), NumDisplayedStats );
+			Text = FText::Format(LOCTEXT("GroupAndStat_GroupNodeTextFmt", "{0} ({1}) ({2})"), FText::FromName(GroupOrStatNode->GetName()), FText::AsNumber(GroupOrStatNode->GetChildren().Num()), FText::AsNumber(NumDisplayedStats));
 		}
 		else
 		{
-			const bool bIsStatTracked = FProfilerManager::Get()->IsStatTracked( GroupOrStatNode->GetStatID() ); 
-			Text = FString::Printf( TEXT("%s %s"), *GroupOrStatNode->GetName().GetPlainNameString(), bIsStatTracked ? TEXT(" *") : TEXT("") );
+			const bool bIsStatTracked = FProfilerManager::Get()->IsStatTracked( GroupOrStatNode->GetStatID() );
+			if(bIsStatTracked)
+			{
+				Text = FText::Format(LOCTEXT("GroupAndStat_GroupNodeTrackedTextFmt", "{0}*"), FText::FromName(GroupOrStatNode->GetName()));
+			}
+			else
+			{
+				Text = FText::FromName(GroupOrStatNode->GetName());
+			}
 		}
 
 		return Text;
@@ -614,7 +621,7 @@ void SFiltersAndPresets::Construct( const FArguments& InArgs )
 					.HintText( LOCTEXT("SearchBoxHint", "Search stats or groups") )
 					.OnTextChanged( this, &SFiltersAndPresets::SearchBox_OnTextChanged )
 					.IsEnabled( this, &SFiltersAndPresets::SearchBox_IsEnabled )
-					.ToolTipText( LOCTEXT("FilterSearchHint", "Type here to search stats or group").ToString() )
+					.ToolTipText( LOCTEXT("FilterSearchHint", "Type here to search stats or group") )
 				]
 
 				// Group by and Sort By
@@ -630,7 +637,7 @@ void SFiltersAndPresets::Construct( const FArguments& InArgs )
 					.VAlign(VAlign_Center)
 					[
 						SNew(STextBlock)
-						.Text( LOCTEXT("GroupByText", "Group by").ToString() )
+						.Text( LOCTEXT("GroupByText", "Group by") )
 					]
 
 					+SHorizontalBox::Slot()
@@ -662,7 +669,7 @@ void SFiltersAndPresets::Construct( const FArguments& InArgs )
 					.VAlign(VAlign_Center)
 					[
 						SNew(STextBlock)
-						.Text( LOCTEXT("SortByText", "Sort by").ToString() )
+						.Text( LOCTEXT("SortByText", "Sort by") )
 					]
 
 					+SHorizontalBox::Slot()
@@ -1047,7 +1054,7 @@ TSharedRef<SWidget> SFiltersAndPresets::GetToggleButtonForStatType( const EProfi
 		.Padding( 2.0f )
 		.OnCheckStateChanged( this, &SFiltersAndPresets::FilterByStatType_OnCheckStateChanged, StatType )
 		.IsChecked( this, &SFiltersAndPresets::FilterByStatType_IsChecked, StatType )
-		.ToolTipText( EProfilerSampleTypes::ToDescription( StatType ) )
+		.ToolTipText( FText::FromString(EProfilerSampleTypes::ToDescription( StatType )) )
 		[
 			SNew(SHorizontalBox)
 
@@ -1064,7 +1071,7 @@ TSharedRef<SWidget> SFiltersAndPresets::GetToggleButtonForStatType( const EProfi
 				.VAlign( VAlign_Center )
 				[
 					SNew( STextBlock )
-						.Text( EProfilerSampleTypes::ToName( StatType ) )
+						.Text( FText::FromString(EProfilerSampleTypes::ToName( StatType )) )
 						.TextStyle( FEditorStyle::Get(), TEXT("Profiler.Caption") )
 				]
 		];
diff --git a/Engine/Source/Developer/Profiler/Private/Widgets/SProfilerFPSChartPanel.cpp b/Engine/Source/Developer/Profiler/Private/Widgets/SProfilerFPSChartPanel.cpp
--- a/Engine/Source/Developer/Profiler/Private/Widgets/SProfilerFPSChartPanel.cpp
+++ b/Engine/Source/Developer/Profiler/Private/Widgets/SProfilerFPSChartPanel.cpp
@@ -89,40 +89,51 @@ public:
 	END_SLATE_FUNCTION_BUILD_OPTIMIZATION
 
 protected:
-	FString HandleSampleCount() const
+	FText HandleSampleCount() const
 	{
-		FString Text = FString::Printf(TEXT("Samples: %6d"), FPSAnalyzer->Samples.Num());
-		return Text;
+		return FText::Format(LOCTEXT("SamplesCountFmt", "Samples: {0}"), FText::AsNumber(FPSAnalyzer->Samples.Num()));
 	}
-	FString HandleMinFPS() const
+	FText HandleMinFPS() const
 	{
-		FString Text = FPSAnalyzer->Samples.Num() > 0 ? FString::Printf(TEXT("Min FPS: %3.2f"), FPSAnalyzer->MinFPS) : TEXT("Min FPS: ");
-		return Text;
+		static const FNumberFormattingOptions FormatOptions = FNumberFormattingOptions()
+			.SetMinimumFractionalDigits(2)
+			.SetMaximumFractionalDigits(2);
+		return FText::Format(LOCTEXT("MinFPSFmt", "Min FPS: {0}"), ((FPSAnalyzer->Samples.Num() > 0) ? FText::AsNumber(FPSAnalyzer->MinFPS, &FormatOptions) : FText::GetEmpty()));
 	}
-	FString HandleMaxFPS() const
+	FText HandleMaxFPS() const
 	{
-		FString Text = FPSAnalyzer->Samples.Num() > 0 ? FString::Printf(TEXT("Max FPS: %3.2f"), FPSAnalyzer->MaxFPS) : TEXT("Max FPS: ");
-		return Text;
+		static const FNumberFormattingOptions FormatOptions = FNumberFormattingOptions()
+			.SetMinimumFractionalDigits(2)
+			.SetMaximumFractionalDigits(2);
+		return FText::Format(LOCTEXT("MaxFPSFmt", "Max FPS: {0}"), ((FPSAnalyzer->Samples.Num() > 0) ? FText::AsNumber(FPSAnalyzer->MaxFPS, &FormatOptions) : FText::GetEmpty()));
 	}
-	FString HandleAverageFPS() const
+	FText HandleAverageFPS() const
 	{
-		FString Text = FPSAnalyzer->Samples.Num() > 0 ? FString::Printf(TEXT("Ave FPS: %3.2f"), FPSAnalyzer->AveFPS) : TEXT("Ave FPS: ");
-		return Text;
+		static const FNumberFormattingOptions FormatOptions = FNumberFormattingOptions()
+			.SetMinimumFractionalDigits(2)
+			.SetMaximumFractionalDigits(2);
+		return FText::Format(LOCTEXT("AverageFPSFmt", "Ave FPS: {0}"), ((FPSAnalyzer->Samples.Num() > 0) ? FText::AsNumber(FPSAnalyzer->AveFPS, &FormatOptions) : FText::GetEmpty()));
 	}
-	FString HandleFPS30() const
+	FText HandleFPS30() const
 	{
-		FString Text = FPSAnalyzer->Samples.Num() > 0 ? FString::Printf(TEXT("+30FPS: %3.2f"), 100.0f*(float)FPSAnalyzer->FPS30 / (float)FPSAnalyzer->Samples.Num()) : TEXT("+30FPS: ");
-		return Text;
+		static const FNumberFormattingOptions FormatOptions = FNumberFormattingOptions()
+			.SetMinimumFractionalDigits(2)
+			.SetMaximumFractionalDigits(2);
+		return FText::Format(LOCTEXT("+30FPSFmt", "+30FPS: {0}"), ((FPSAnalyzer->Samples.Num() > 0) ? FText::AsNumber(100.0f*(float)FPSAnalyzer->FPS30 / (float)FPSAnalyzer->Samples.Num(), &FormatOptions) : FText::GetEmpty()));
 	}
-	FString HandleFPS25() const
+	FText HandleFPS25() const
 	{
-		FString Text = FPSAnalyzer->Samples.Num() > 0 ? FString::Printf(TEXT("+25FPS: %3.2f"), 100.0f*(float)FPSAnalyzer->FPS25 / (float)FPSAnalyzer->Samples.Num()) : TEXT("+25FPS: ");
-		return Text;
+		static const FNumberFormattingOptions FormatOptions = FNumberFormattingOptions()
+			.SetMinimumFractionalDigits(2)
+			.SetMaximumFractionalDigits(2);
+		return FText::Format(LOCTEXT("+25FPSFmt", "+25FPS: {0}"), ((FPSAnalyzer->Samples.Num() > 0) ? FText::AsNumber(100.0f*(float)FPSAnalyzer->FPS25 / (float)FPSAnalyzer->Samples.Num(), &FormatOptions) : FText::GetEmpty()));
 	}
-	FString HandleFPS20() const
+	FText HandleFPS20() const
 	{
-		FString Text = FPSAnalyzer->Samples.Num() > 0 ? FString::Printf(TEXT("+20FPS: %3.2f"), 100.0f*(float)FPSAnalyzer->FPS20 / (float)FPSAnalyzer->Samples.Num()) : TEXT("+20FPS: ");
-		return Text;
+		static const FNumberFormattingOptions FormatOptions = FNumberFormattingOptions()
+			.SetMinimumFractionalDigits(2)
+			.SetMaximumFractionalDigits(2);
+		return FText::Format(LOCTEXT("+20FPSFmt", "+20FPS: {0}"), ((FPSAnalyzer->Samples.Num() > 0) ? FText::AsNumber(100.0f*(float)FPSAnalyzer->FPS20 / (float)FPSAnalyzer->Samples.Num(), &FormatOptions) : FText::GetEmpty()));
 	}
 
 protected:
diff --git a/Engine/Source/Developer/Profiler/Private/Widgets/SProfilerToolbar.cpp b/Engine/Source/Developer/Profiler/Private/Widgets/SProfilerToolbar.cpp
--- a/Engine/Source/Developer/Profiler/Private/Widgets/SProfilerToolbar.cpp
+++ b/Engine/Source/Developer/Profiler/Private/Widgets/SProfilerToolbar.cpp
@@ -220,7 +220,7 @@ void SProfilerToolbar::AddSessionInstanceItem( TSharedRef<SHorizontalBox>& Sessi
 		.ButtonContent()
 		[
 			SNew(STextBlock)
-			.Text( ProfilerSessionName )
+			.Text( FText::FromString(ProfilerSessionName) )
 			.ShadowOffset( FVector2D(1.0f,1.0f) )
 			.ShadowColorAndOpacity( FColorList::SkyBlue )
 		]
diff --git a/Engine/Source/Developer/Profiler/Private/Widgets/SProfilerWindow.cpp b/Engine/Source/Developer/Profiler/Private/Widgets/SProfilerWindow.cpp
--- a/Engine/Source/Developer/Profiler/Private/Widgets/SProfilerWindow.cpp
+++ b/Engine/Source/Developer/Profiler/Private/Widgets/SProfilerWindow.cpp
@@ -339,7 +339,7 @@ void SProfilerWindow::ManageEventGraphTab( const FGuid ProfilerInstanceID, const
 			.AutoWidth()
 			[
 				SNew(STextBlock)
-				.Text( TabName )
+				.Text( FText::FromString(TabName) )
 			]
 		];
 
diff --git a/Engine/Source/Developer/Profiler/Private/Widgets/StatDragDropOp.h b/Engine/Source/Developer/Profiler/Private/Widgets/StatDragDropOp.h
--- a/Engine/Source/Developer/Profiler/Private/Widgets/StatDragDropOp.h
+++ b/Engine/Source/Developer/Profiler/Private/Widgets/StatDragDropOp.h
@@ -87,7 +87,7 @@ private:
 			.VAlign( VAlign_Center )
 			[
 				SNew(STextBlock)
-				.Text( Description )
+				.Text( FText::FromString(Description) )
 			]
 		];
 	}
diff --git a/Engine/Source/Developer/ProjectLauncher/Private/Widgets/Build/SProjectLauncherBuildPage.cpp b/Engine/Source/Developer/ProjectLauncher/Private/Widgets/Build/SProjectLauncherBuildPage.cpp
--- a/Engine/Source/Developer/ProjectLauncher/Private/Widgets/Build/SProjectLauncherBuildPage.cpp
+++ b/Engine/Source/Developer/ProjectLauncher/Private/Widgets/Build/SProjectLauncherBuildPage.cpp
@@ -276,16 +276,16 @@ void SProjectLauncherBuildPage::HandleBuildConfigurationSelectorConfigurationSel
 }
 
 
-FString SProjectLauncherBuildPage::HandleBuildConfigurationSelectorText() const
+FText SProjectLauncherBuildPage::HandleBuildConfigurationSelectorText() const
 {
 	ILauncherProfilePtr SelectedProfile = Model->GetSelectedProfile();
 
 	if (SelectedProfile.IsValid())
 	{
-		return EBuildConfigurations::ToString(SelectedProfile->GetBuildConfiguration());
+		return FText::FromString(EBuildConfigurations::ToString(SelectedProfile->GetBuildConfiguration()));
 	}
 
-	return FString();
+	return FText::GetEmpty();
 }
 
 
diff --git a/Engine/Source/Developer/ProjectLauncher/Private/Widgets/Build/SProjectLauncherBuildPage.h b/Engine/Source/Developer/ProjectLauncher/Private/Widgets/Build/SProjectLauncherBuildPage.h
--- a/Engine/Source/Developer/ProjectLauncher/Private/Widgets/Build/SProjectLauncherBuildPage.h
+++ b/Engine/Source/Developer/ProjectLauncher/Private/Widgets/Build/SProjectLauncherBuildPage.h
@@ -67,7 +67,7 @@ private:
 	void HandleBuildConfigurationSelectorConfigurationSelected(EBuildConfigurations::Type Configuration);
 
 	/** Callback for getting the content text of the build configuration selector. */
-	FString HandleBuildConfigurationSelectorText() const;
+	FText HandleBuildConfigurationSelectorText() const;
 
 	/** Callback for determining the visibility of a validation error icon. */
 	EVisibility HandleValidationErrorIconVisibility(ELauncherProfileValidationErrors::Type Error) const;
diff --git a/Engine/Source/Developer/ProjectLauncher/Private/Widgets/Cook/SProjectLauncherCookByTheBookSettings.cpp b/Engine/Source/Developer/ProjectLauncher/Private/Widgets/Cook/SProjectLauncherCookByTheBookSettings.cpp
--- a/Engine/Source/Developer/ProjectLauncher/Private/Widgets/Cook/SProjectLauncherCookByTheBookSettings.cpp
+++ b/Engine/Source/Developer/ProjectLauncher/Private/Widgets/Cook/SProjectLauncherCookByTheBookSettings.cpp
@@ -843,16 +843,16 @@ void SProjectLauncherCookByTheBookSettings::HandleCookConfigurationSelectorConfi
 }
 
 
-FString SProjectLauncherCookByTheBookSettings::HandleCookConfigurationSelectorText( ) const
+FText SProjectLauncherCookByTheBookSettings::HandleCookConfigurationSelectorText( ) const
 {
 	ILauncherProfilePtr SelectedProfile = Model->GetSelectedProfile();
 
 	if (SelectedProfile.IsValid())
 	{
-		return EBuildConfigurations::ToString(SelectedProfile->GetCookConfiguration());
+		return FText::FromString(EBuildConfigurations::ToString(SelectedProfile->GetCookConfiguration()));
 	}
 
-	return TEXT("");
+	return FText::GetEmpty();
 }
 
 
diff --git a/Engine/Source/Developer/ProjectLauncher/Private/Widgets/Cook/SProjectLauncherCookByTheBookSettings.h b/Engine/Source/Developer/ProjectLauncher/Private/Widgets/Cook/SProjectLauncherCookByTheBookSettings.h
--- a/Engine/Source/Developer/ProjectLauncher/Private/Widgets/Cook/SProjectLauncherCookByTheBookSettings.h
+++ b/Engine/Source/Developer/ProjectLauncher/Private/Widgets/Cook/SProjectLauncherCookByTheBookSettings.h
@@ -65,7 +65,7 @@ private:
 	void HandleCookConfigurationSelectorConfigurationSelected( EBuildConfigurations::Type );
 
 	/** Handles getting the content text of the cooker build configuration selector. */
-	FString HandleCookConfigurationSelectorText( ) const;
+	FText HandleCookConfigurationSelectorText( ) const;
 
 	/** Handles check state changes of the 'Incremental' check box. */
 	void HandleIncrementalCheckBoxCheckStateChanged( ECheckBoxState NewState );
diff --git a/Engine/Source/Developer/ProjectLauncher/Private/Widgets/Cook/SProjectLauncherCultureListRow.h b/Engine/Source/Developer/ProjectLauncher/Private/Widgets/Cook/SProjectLauncherCultureListRow.h
--- a/Engine/Source/Developer/ProjectLauncher/Private/Widgets/Cook/SProjectLauncherCultureListRow.h
+++ b/Engine/Source/Developer/ProjectLauncher/Private/Widgets/Cook/SProjectLauncherCultureListRow.h
@@ -56,7 +56,7 @@ public:
 				.Padding(FMargin(6.0, 2.0))
 				[
 					SNew(STextBlock)
-						.Text(*CultureName)			
+						.Text(FText::FromString(*CultureName))
 				];
 		}
 
diff --git a/Engine/Source/Developer/ProjectLauncher/Private/Widgets/Cook/SProjectLauncherMapListRow.h b/Engine/Source/Developer/ProjectLauncher/Private/Widgets/Cook/SProjectLauncherMapListRow.h
--- a/Engine/Source/Developer/ProjectLauncher/Private/Widgets/Cook/SProjectLauncherMapListRow.h
+++ b/Engine/Source/Developer/ProjectLauncher/Private/Widgets/Cook/SProjectLauncherMapListRow.h
@@ -55,7 +55,7 @@ public:
 				.Padding(FMargin(6.0, 2.0))
 				[
 					SNew(STextBlock)
-						.Text(*MapName)			
+						.Text(FText::FromString(*MapName))
 				];
 		}
 
diff --git a/Engine/Source/Developer/ProjectLauncher/Private/Widgets/Cook/SProjectLauncherPlatformListRow.h b/Engine/Source/Developer/ProjectLauncher/Private/Widgets/Cook/SProjectLauncherPlatformListRow.h
--- a/Engine/Source/Developer/ProjectLauncher/Private/Widgets/Cook/SProjectLauncherPlatformListRow.h
+++ b/Engine/Source/Developer/ProjectLauncher/Private/Widgets/Cook/SProjectLauncherPlatformListRow.h
@@ -55,7 +55,7 @@ public:
 				.Padding(FMargin(6.0, 2.0))
 				[
 					SNew(STextBlock)
-						.Text(*PlatformName)			
+						.Text(FText::FromString(*PlatformName))
 				];
 		}
 
diff --git a/Engine/Source/Developer/ProjectLauncher/Private/Widgets/Deploy/SProjectLauncherDeployPage.cpp b/Engine/Source/Developer/ProjectLauncher/Private/Widgets/Deploy/SProjectLauncherDeployPage.cpp
--- a/Engine/Source/Developer/ProjectLauncher/Private/Widgets/Deploy/SProjectLauncherDeployPage.cpp
+++ b/Engine/Source/Developer/ProjectLauncher/Private/Widgets/Deploy/SProjectLauncherDeployPage.cpp
@@ -144,7 +144,7 @@ END_SLATE_FUNCTION_BUILD_OPTIMIZATION
 /* SProjectLauncherDevicesPage callbacks
  *****************************************************************************/
 
-FString SProjectLauncherDeployPage::HandleDeploymentModeComboButtonContentText( ) const
+FText SProjectLauncherDeployPage::HandleDeploymentModeComboButtonContentText( ) const
 {
 	ILauncherProfilePtr SelectedProfile = Model->GetSelectedProfile();
 
@@ -154,28 +154,28 @@ FString SProjectLauncherDeployPage::HandleDeploymentModeComboButtonContentText(
 
 		if (DeploymentMode == ELauncherProfileDeploymentModes::CopyToDevice)
 		{
-			return TEXT("Copy to device");
+			return LOCTEXT("CopyToDeviceAction", "Copy to device");
 		}
 
 		if (DeploymentMode == ELauncherProfileDeploymentModes::DoNotDeploy)
 		{
-			return TEXT("Do not deploy");
+			return LOCTEXT("DoNotDeployAction", "Do not deploy");
 		}
 
 		if (DeploymentMode == ELauncherProfileDeploymentModes::FileServer)
 		{
-			return TEXT("File server");
+			return LOCTEXT("FileServerAction", "File server");
 		}
 
 		if (DeploymentMode == ELauncherProfileDeploymentModes::CopyRepository)
 		{
-			return TEXT("Copy repository");
+			return LOCTEXT("CopyRepositoryAction", "Copy repository");
 		}
 
-		return LOCTEXT("DeploymentModeComboButtonDefaultText", "Select...").ToString();
+		return LOCTEXT("DeploymentModeComboButtonDefaultText", "Select...");
 	}
 
-	return TEXT("");
+	return FText::GetEmpty();
 }
 
 
diff --git a/Engine/Source/Developer/ProjectLauncher/Private/Widgets/Deploy/SProjectLauncherDeployPage.h b/Engine/Source/Developer/ProjectLauncher/Private/Widgets/Deploy/SProjectLauncherDeployPage.h
--- a/Engine/Source/Developer/ProjectLauncher/Private/Widgets/Deploy/SProjectLauncherDeployPage.h
+++ b/Engine/Source/Developer/ProjectLauncher/Private/Widgets/Deploy/SProjectLauncherDeployPage.h
@@ -34,7 +34,7 @@ public:
 private:
 
 	// Callback for getting the content text of the 'Deployment Mode' combo button.
-	FString HandleDeploymentModeComboButtonContentText( ) const;
+	FText HandleDeploymentModeComboButtonContentText( ) const;
 
 	// Callback for clicking an item in the 'Deployment Mode' menu.
 	void HandleDeploymentModeMenuEntryClicked( ELauncherProfileDeploymentModes::Type DeploymentMode );
diff --git a/Engine/Source/Developer/ProjectLauncher/Private/Widgets/Deploy/SProjectLauncherDeployTargetListRow.h b/Engine/Source/Developer/ProjectLauncher/Private/Widgets/Deploy/SProjectLauncherDeployTargetListRow.h
--- a/Engine/Source/Developer/ProjectLauncher/Private/Widgets/Deploy/SProjectLauncherDeployTargetListRow.h
+++ b/Engine/Source/Developer/ProjectLauncher/Private/Widgets/Deploy/SProjectLauncherDeployTargetListRow.h
@@ -86,7 +86,7 @@ public:
 			return SNew(SCheckBox)
 				.IsChecked(this, &SProjectLauncherDeployTargetListRow::HandleCheckBoxIsChecked)
 				.OnCheckStateChanged(this, &SProjectLauncherDeployTargetListRow::HandleCheckBoxStateChanged)
-				.ToolTipText(LOCTEXT("CheckBoxToolTip", "Check this box to include this device in the current device group").ToString());
+				.ToolTipText(LOCTEXT("CheckBoxToolTip", "Check this box to include this device in the current device group"));
 		}
 		else if (ColumnName == "Device")
 		{
@@ -206,19 +206,19 @@ private:
 	}
 
 	// Callback for getting the text of the current variant
-	FString HandleVariantSelectorText() const
+	FText HandleVariantSelectorText() const
 	{
 		ILauncherDeviceGroupPtr ActiveGroup = DeviceGroup.Get();
 		if (ActiveGroup.IsValid() && DeviceProxy.IsValid())
 		{
 			if (SelectedVariant == NAME_None)
 			{
-				return "Default";
+				return LOCTEXT("DefaultVariant", "Default");
 			}
-			return SelectedVariant.ToString();
+			return FText::FromName(SelectedVariant);
 		}
 
-		return FString();
+		return FText::GetEmpty();
 	}
 
 	// Callback for changing the selected variant
@@ -257,44 +257,44 @@ private:
 	}
 
 	// Callback for getting the friendly name.
-	FString HandleDeviceNameText() const
+	FText HandleDeviceNameText() const
 	{
 		const FString& Name = DeviceProxy->GetName();
 
 		if (Name.IsEmpty())
 		{
-			return LOCTEXT("UnnamedDeviceName", "<unnamed>").ToString();
+			return LOCTEXT("UnnamedDeviceName", "<unnamed>");
 		}
 
-		return Name;
+		return FText::FromString(Name);
 	}
 
 	// Callback for getting the host name.
-	FString HandleHostNameText() const
+	FText HandleHostNameText() const
 	{
-		return DeviceProxy->GetHostName();
+		return FText::FromString(DeviceProxy->GetHostName());
 	}
 
 	// Callback for getting the host user name.
-	FString HandleHostUserText() const
+	FText HandleHostUserText() const
 	{
-		return DeviceProxy->GetHostUser();
+		return FText::FromString(DeviceProxy->GetHostUser());
 	}
 
 	// Callback for getting the host platform name.
-	FString HandleHostPlatformText() const
+	FText HandleHostPlatformText() const
 	{
 		if (DeviceProxy->HasVariant(NAME_None))
 		{
-			return DeviceProxy->GetTargetPlatformName(NAME_None);
+			return FText::FromString(DeviceProxy->GetTargetPlatformName(NAME_None));
 		}
-		return LOCTEXT("InvalidVariant", "Invalid Variant").ToString();
+		return LOCTEXT("InvalidVariant", "Invalid Variant");
 	}
 
 	// Callback for getting the default variant name in the case where 
-	FString HandleHostNoVariantText() const
+	FText HandleHostNoVariantText() const
 	{
-		return LOCTEXT("StandardVariant", "Standard").ToString();
+		return LOCTEXT("StandardVariant", "Standard");
 	}
 
 private:
diff --git a/Engine/Source/Developer/ProjectLauncher/Private/Widgets/Deploy/SProjectLauncherDeviceGroupSelector.cpp b/Engine/Source/Developer/ProjectLauncher/Private/Widgets/Deploy/SProjectLauncherDeviceGroupSelector.cpp
--- a/Engine/Source/Developer/ProjectLauncher/Private/Widgets/Deploy/SProjectLauncherDeviceGroupSelector.cpp
+++ b/Engine/Source/Developer/ProjectLauncher/Private/Widgets/Deploy/SProjectLauncherDeviceGroupSelector.cpp
@@ -53,16 +53,16 @@ void SProjectLauncherDeviceGroupSelector::Construct(const FArguments& InArgs, co
 	ProfileManager->OnDeviceGroupRemoved().AddSP(this, &SProjectLauncherDeviceGroupSelector::HandleProfileManagerDeviceGroupsChanged);
 }
 
-FString SProjectLauncherDeviceGroupSelector::HandleDeviceGroupComboBoxContent() const
+FText SProjectLauncherDeviceGroupSelector::HandleDeviceGroupComboBoxContent() const
 {
 	ILauncherDeviceGroupPtr SelectedGroup = DeviceGroupComboBox->GetSelectedItem();
 
 	if (SelectedGroup.IsValid())
 	{
-		return SelectedGroup->GetName();
+		return FText::FromString(SelectedGroup->GetName());
 	}
 
-	return LOCTEXT("CreateOrSelectGroupText", "Create or select a device group...").ToString();
+	return LOCTEXT("CreateOrSelectGroupText", "Create or select a device group...");
 }
 
 TSharedRef<SWidget> SProjectLauncherDeviceGroupSelector::HandleDeviceGroupComboBoxGenerateWidget(ILauncherDeviceGroupPtr InItem)
@@ -104,14 +104,14 @@ FReply SProjectLauncherDeviceGroupSelector::HandleDeviceGroupComboBoxRemoveClick
 	return FReply::Handled();
 }
 
-FString SProjectLauncherDeviceGroupSelector::HandleDeviceGroupComboBoxWidgetText(ILauncherDeviceGroupPtr Group) const
+FText SProjectLauncherDeviceGroupSelector::HandleDeviceGroupComboBoxWidgetText(ILauncherDeviceGroupPtr Group) const
 {
 	if (Group.IsValid())
 	{
-		return Group->GetName();
+		return FText::FromString(Group->GetName());
 	}
 
-	return FString();
+	return FText::GetEmpty();
 }
 
 void SProjectLauncherDeviceGroupSelector::HandleProfileManagerDeviceGroupsChanged(const ILauncherDeviceGroupRef& /*ChangedProfile*/)
diff --git a/Engine/Source/Developer/ProjectLauncher/Private/Widgets/Deploy/SProjectLauncherDeviceGroupSelector.h b/Engine/Source/Developer/ProjectLauncher/Private/Widgets/Deploy/SProjectLauncherDeviceGroupSelector.h
--- a/Engine/Source/Developer/ProjectLauncher/Private/Widgets/Deploy/SProjectLauncherDeviceGroupSelector.h
+++ b/Engine/Source/Developer/ProjectLauncher/Private/Widgets/Deploy/SProjectLauncherDeviceGroupSelector.h
@@ -81,7 +81,7 @@ private:
 	FReply HandleDeviceGroupComboBoxAddClicked( );
 
 	// Callback for getting the device group combo box content.
-	FString HandleDeviceGroupComboBoxContent() const;
+	FText HandleDeviceGroupComboBoxContent() const;
 
 	// Callback for generating a row for the device group combo box.
 	TSharedRef<SWidget> HandleDeviceGroupComboBoxGenerateWidget( ILauncherDeviceGroupPtr InItem );
@@ -102,7 +102,7 @@ private:
 	void HandleDeviceGroupComboBoxSelectionRenamed( const FText& CommittedText, ETextCommit::Type );
 
 	// Callback for getting the text of a device group combo box widget.
-	FString HandleDeviceGroupComboBoxWidgetText( ILauncherDeviceGroupPtr Group ) const;
+	FText HandleDeviceGroupComboBoxWidgetText( ILauncherDeviceGroupPtr Group ) const;
 
 	// Callback for changing the list of groups in the profile manager.
 	void HandleProfileManagerDeviceGroupsChanged( const ILauncherDeviceGroupRef& /*ChangedProfile*/ );
diff --git a/Engine/Source/Developer/ProjectLauncher/Private/Widgets/Deploy/SProjectLauncherSimpleDeviceListRow.h b/Engine/Source/Developer/ProjectLauncher/Private/Widgets/Deploy/SProjectLauncherSimpleDeviceListRow.h
--- a/Engine/Source/Developer/ProjectLauncher/Private/Widgets/Deploy/SProjectLauncherSimpleDeviceListRow.h
+++ b/Engine/Source/Developer/ProjectLauncher/Private/Widgets/Deploy/SProjectLauncherSimpleDeviceListRow.h
@@ -211,19 +211,19 @@ private:
 			];
 	}
 
-	FString HandleVariantSelectorText() const
+	FText HandleVariantSelectorText() const
 	{
 		if (SimpleProfile.IsValid())
 		{
 			FName Variant = SimpleProfile->GetDeviceVariant();
 			if (Variant == NAME_None)
 			{
-				return "Default";
+				return LOCTEXT("DefaultVariant", "Default");
 			}
-			return Variant.ToString();
+			return FText::FromName(Variant);
 		}
 
-		return FString();
+		return FText::GetEmpty();
 	}
 
 	void HandleVariantSelectorVariantSelected(FName InVariant)
@@ -258,14 +258,14 @@ private:
 			];
 	}
 
-	FString HandleBuildConfigurationSelectorText() const
+	FText HandleBuildConfigurationSelectorText() const
 	{
 		if (SimpleProfile.IsValid())
 		{
-			return EBuildConfigurations::ToString(SimpleProfile->GetBuildConfiguration());
+			return FText::FromString(EBuildConfigurations::ToString(SimpleProfile->GetBuildConfiguration()));
 		}
 
-		return FString();
+		return FText::GetEmpty();
 	}
 
 	void HandleBuildConfigurationSelectorConfigurationSelected(EBuildConfigurations::Type Configuration)
@@ -294,33 +294,33 @@ private:
 			.HAlign(HAlign_Fill)
 			[
 				SNew(SProjectLauncherCookModeSelector)
-				.OnCookModeSelected(this, &SProjectLauncherSimpleDeviceListRow::HandleCookModSelectorConfigurationSelected)
-				.Text(this, &SProjectLauncherSimpleDeviceListRow::HandleCookModSelectorText)
+				.OnCookModeSelected(this, &SProjectLauncherSimpleDeviceListRow::HandleCookModeSelectorConfigurationSelected)
+				.Text(this, &SProjectLauncherSimpleDeviceListRow::HandleCookModeSelectorText)
 				.Visibility(this, &SProjectLauncherSimpleDeviceListRow::IsAdvancedVisible)
 			];
 	}
 
-	FString HandleCookModSelectorText() const
+	FText HandleCookModeSelectorText() const
 	{
 		if (SimpleProfile.IsValid())
 		{
 			switch (SimpleProfile->GetCookMode())
 			{
 			case ELauncherProfileCookModes::DoNotCook:
-				return TEXT("Do not cook");
+				return LOCTEXT("CookMode_DoNotCook", "Do not cook");
 
 			case ELauncherProfileCookModes::ByTheBook:
-				return TEXT("By the book");
+				return LOCTEXT("CookMode_ByTheBook", "By the book");
 
 			case ELauncherProfileCookModes::OnTheFly:
-				return TEXT("On the fly");
+				return LOCTEXT("CookMode_OnTheFly", "On the fly");
 			}
 		}
 
-		return FString();
+		return FText::GetEmpty();
 	}
 
-	void HandleCookModSelectorConfigurationSelected(ELauncherProfileCookModes::Type CookMode)
+	void HandleCookModeSelectorConfigurationSelected(ELauncherProfileCookModes::Type CookMode)
 	{
 		if (SimpleProfile.IsValid())
 		{
@@ -408,38 +408,38 @@ private:
 	}
 
 	// Callback for getting the friendly name.
-	FString HandleDeviceNameText() const
+	FText HandleDeviceNameText() const
 	{
 		const FString& Name = DeviceProxy->GetName();
 
 		if (Name.IsEmpty())
 		{
-			return LOCTEXT("UnnamedDeviceName", "<unnamed>").ToString();
+			return LOCTEXT("UnnamedDeviceName", "<unnamed>");
 		}
 
-		return Name;
+		return FText::FromString(Name);
 	}
 
 	// Callback for getting the host name.
-	FString HandleHostNameText( ) const
+	FText HandleHostNameText( ) const
 	{
-		return DeviceProxy->GetHostName();
+		return FText::FromString(DeviceProxy->GetHostName());
 	}
 
 	// Callback for getting the host user name.
-	FString HandleHostUserText( ) const
+	FText HandleHostUserText( ) const
 	{
-		return DeviceProxy->GetHostUser();
+		return FText::FromString(DeviceProxy->GetHostUser());
 	}
 
 	// Callback for getting the host platform name.
-	FString HandleHostPlatformText( ) const
+	FText HandleHostPlatformText( ) const
 	{
 		if (LaunchProfile.IsValid())
 		{
-			return DeviceProxy->GetTargetPlatformName(SimpleProfile->GetDeviceVariant());
+			return FText::FromString(DeviceProxy->GetTargetPlatformName(SimpleProfile->GetDeviceVariant()));
 		}
-		return LOCTEXT("InvalidVariant", "Invalid Variant").ToString();
+		return LOCTEXT("InvalidVariant", "Invalid Variant");
 	}
 
 private:
diff --git a/Engine/Source/Developer/ProjectLauncher/Private/Widgets/Launch/SProjectLauncherLaunchPage.cpp b/Engine/Source/Developer/ProjectLauncher/Private/Widgets/Launch/SProjectLauncherLaunchPage.cpp
--- a/Engine/Source/Developer/ProjectLauncher/Private/Widgets/Launch/SProjectLauncherLaunchPage.cpp
+++ b/Engine/Source/Developer/ProjectLauncher/Private/Widgets/Launch/SProjectLauncherLaunchPage.cpp
@@ -211,7 +211,7 @@ EVisibility SProjectLauncherLaunchPage::HandleLaunchModeBoxVisibility( ) const
 }
 
 
-FString SProjectLauncherLaunchPage::HandleLaunchModeComboButtonContentText( ) const
+FText SProjectLauncherLaunchPage::HandleLaunchModeComboButtonContentText( ) const
 {
 	ILauncherProfilePtr SelectedProfile = Model->GetSelectedProfile();
 
@@ -221,23 +221,23 @@ FString SProjectLauncherLaunchPage::HandleLaunchModeComboButtonContentText( ) co
 
 		if (LaunchMode == ELauncherProfileLaunchModes::CustomRoles)
 		{
-			return LOCTEXT("LaunchModeComboButtonCustomRolesText", "Using custom roles").ToString();
+			return LOCTEXT("LaunchModeComboButtonCustomRolesText", "Using custom roles");
 		}
 
 		if (LaunchMode == ELauncherProfileLaunchModes::DefaultRole)
 		{
-			return LOCTEXT("LaunchModeComboButtonDefaultRoleText", "Using default role").ToString();
+			return LOCTEXT("LaunchModeComboButtonDefaultRoleText", "Using default role");
 		}
 
 		if (LaunchMode == ELauncherProfileLaunchModes::DoNotLaunch)
 		{
-			return LOCTEXT("LaunchModeComboButtonDoNotLaunchText", "Do not launch").ToString();
+			return LOCTEXT("LaunchModeComboButtonDoNotLaunchText", "Do not launch");
 		}
 
-		return LOCTEXT("LaunchModeComboButtonSelectText", "Select...").ToString();
+		return LOCTEXT("LaunchModeComboButtonSelectText", "Select...");
 	}
 
-	return TEXT("");
+	return FText::GetEmpty();
 }
 
 
diff --git a/Engine/Source/Developer/ProjectLauncher/Private/Widgets/Launch/SProjectLauncherLaunchPage.h b/Engine/Source/Developer/ProjectLauncher/Private/Widgets/Launch/SProjectLauncherLaunchPage.h
--- a/Engine/Source/Developer/ProjectLauncher/Private/Widgets/Launch/SProjectLauncherLaunchPage.h
+++ b/Engine/Source/Developer/ProjectLauncher/Private/Widgets/Launch/SProjectLauncherLaunchPage.h
@@ -45,7 +45,7 @@ private:
 	EVisibility HandleLaunchModeBoxVisibility( ) const;
 
 	// Callback for getting the content text of the 'Launch Mode' combo button.
-	FString HandleLaunchModeComboButtonContentText( ) const;
+	FText HandleLaunchModeComboButtonContentText( ) const;
 
 	// Callback for clicking an item in the 'Launch Mode' menu.
 	void HandleLaunchModeMenuEntryClicked( ELauncherProfileLaunchModes::Type LaunchMode );
