Profiler - Fixed Filter'n'Presets not being updated

diff --git a/Engine/Source/Developer/Profiler/Private/Widgets/SFiltersAndPresets.cpp b/Engine/Source/Developer/Profiler/Private/Widgets/SFiltersAndPresets.cpp
--- a/Engine/Source/Developer/Profiler/Private/Widgets/SFiltersAndPresets.cpp
+++ b/Engine/Source/Developer/Profiler/Private/Widgets/SFiltersAndPresets.cpp
@@ -579,7 +579,6 @@ SFiltersAndPresets::SFiltersAndPresets()
 	: GroupingMode( EStatGroupingOrSortingMode::GroupName )
 	, SortingMode( EStatGroupingOrSortingMode::StatName )
 	, bExpansionSaved( false )
-	, bIsActiveTimerRegistered( false )
 {
 	FMemory::MemSet( bStatTypeIsVisible, 1 );
 }
@@ -761,25 +760,9 @@ void SFiltersAndPresets::Construct( const FArguments& InArgs )
 }
 
 void SFiltersAndPresets::ProfilerManager_OnRequestFilterAndPresetsUpdate()
-{
-	RegisterForUpdate();
-}
-
-void SFiltersAndPresets::RegisterForUpdate()
-{
-	if (!bIsActiveTimerRegistered)
-	{
-		bIsActiveTimerRegistered = true;
-		RegisterActiveTimer(0.f, FWidgetActiveTimerDelegate::CreateSP(this, &SFiltersAndPresets::UpdateData));
-	}
-}
-
-EActiveTimerReturnType SFiltersAndPresets::UpdateData(double InCurrentTime, float InDeltaTime)
 {
 	auto It = FProfilerManager::Get()->GetProfilerInstancesIterator();
 	UpdateGroupAndStatTree(It.Value());
-
-	return EActiveTimerReturnType::Stop;
 }
 
 void SFiltersAndPresets::UpdateGroupAndStatTree( const FProfilerSessionPtr InProfilerSession )
diff --git a/Engine/Source/Developer/Profiler/Private/Widgets/SFiltersAndPresets.h b/Engine/Source/Developer/Profiler/Private/Widgets/SFiltersAndPresets.h
--- a/Engine/Source/Developer/Profiler/Private/Widgets/SFiltersAndPresets.h
+++ b/Engine/Source/Developer/Profiler/Private/Widgets/SFiltersAndPresets.h
@@ -411,14 +411,4 @@ protected:
 
 	/** Holds the visibility of each stat type. */
 	bool bStatTypeIsVisible[EProfilerSampleTypes::InvalidOrMax];
-
-private:
-	/** Registers the active timer to trigger an update if not done already */
-	void RegisterForUpdate();
-
-	/** Simple one-off active timer to update the group and stat trees */
-	EActiveTimerReturnType UpdateData(double InCurrentTime, float InDeltaTime);
-
-	/** Whether the active timer has already been registered this frame */
-	bool bIsActiveTimerRegistered;
 };
