Profiler - Fixed Filters'n'Presets not being updated

diff --git a/Engine/Source/Developer/Profiler/Private/ProfilerManager.cpp b/Engine/Source/Developer/Profiler/Private/ProfilerManager.cpp
--- a/Engine/Source/Developer/Profiler/Private/ProfilerManager.cpp
+++ b/Engine/Source/Developer/Profiler/Private/ProfilerManager.cpp
@@ -336,6 +336,7 @@ void FProfilerManager::ProfilerClient_OnLoadCompleted( const FGuid& InstanceID )
 	const FProfilerSessionRef* ProfilerSession = FindSessionInstance( InstanceID );
 	if( ProfilerSession && GetProfilerWindow().IsValid())
 	{
+		RequestFilterAndPresetsUpdateEvent.Broadcast();
 		(*ProfilerSession)->LoadComplete();
 
 		// Update the notification that a file has been loaded, to be precise it should be loaded on the next tick...
diff --git a/Engine/Source/Developer/Profiler/Private/ProfilerSession.cpp b/Engine/Source/Developer/Profiler/Private/ProfilerSession.cpp
--- a/Engine/Source/Developer/Profiler/Private/ProfilerSession.cpp
+++ b/Engine/Source/Developer/Profiler/Private/ProfilerSession.cpp
@@ -444,6 +444,29 @@ void FProfilerSession::PopulateHierarchy_Recurrent
 	}
 }
 
+void FProfilerSession::LoadComplete()
+{
+	bHasAllProfilerData = true;
+}
+
+void FProfilerSession::UpdateProfilerData( const FProfilerDataFrame& Content )
+{
+	FrameToProfilerDataMapping.FindOrAdd( Content.Frame ) = Content;
+	FrameToProcess.Add( Content.Frame );
+}
+
+void FProfilerSession::UpdateMetadata( const FStatMetaData& InClientStatMetaData )
+{
+	const uint32 CurrentStatMetaDataSize = InClientStatMetaData.GetMetaDataSize();
+	if( CurrentStatMetaDataSize != StatMetaDataSize )
+	{
+		ClientStatMetadata = &InClientStatMetaData;
+		bRequestStatMetadataUpdate = true;
+
+		StatMetaDataSize = CurrentStatMetaDataSize;
+	}
+}
+
 /*-----------------------------------------------------------------------------
 	FRawProfilerSession
 -----------------------------------------------------------------------------*/
@@ -737,6 +760,45 @@ void FRawProfilerSession::PrepareLoading()
 #endif // 0
 }
 
+
+void FProfilerStatMetaData::Update( const FStatMetaData& ClientStatMetaData )
+{
+	PROFILER_SCOPE_LOG_TIME( TEXT( "FProfilerStatMetaData.Update" ), );
+
+	FStatMetaData LocalCopy;
+	{
+		FScopeLock Lock( ClientStatMetaData.CriticalSection );
+		LocalCopy = ClientStatMetaData;
+	}
+
+	// Iterate through all thread descriptions.
+	ThreadDescriptions.Append( LocalCopy.ThreadDescriptions );
+
+	// Initialize fake stat for Self.
+	const uint32 NoGroupID = 0;
+
+	InitializeGroup( NoGroupID, "NoGroup" );
+	InitializeStat( 0, NoGroupID, TEXT( "Self" ), STATTYPE_CycleCounter );
+	InitializeStat( 1, NoGroupID, FStatConstants::NAME_ThreadRoot.GetPlainNameString(), STATTYPE_CycleCounter, FStatConstants::NAME_ThreadRoot );
+
+	// Iterate through all stat group descriptions.
+	for( auto It = LocalCopy.GroupDescriptions.CreateConstIterator(); It; ++It )
+	{
+		const FStatGroupDescription& GroupDesc = It.Value();
+		InitializeGroup( GroupDesc.ID, GroupDesc.Name );
+	}
+
+	// Iterate through all stat descriptions.
+	for( auto It = LocalCopy.StatDescriptions.CreateConstIterator(); It; ++It )
+	{
+		const FStatDescription& StatDesc = It.Value();
+		InitializeStat( StatDesc.ID, StatDesc.GroupID, StatDesc.Name, (EStatType)StatDesc.StatType );
+	}
+
+	SecondsPerCycle = LocalCopy.SecondsPerCycle;
+}
+
+
 void FProfilerStatMetaData::UpdateFromStatsState( const FStatsThreadState& StatsThreadStats )
 {
 	TMap<FName, int32> GroupFNameIDs;
diff --git a/Engine/Source/Developer/Profiler/Private/ProfilerSession.h b/Engine/Source/Developer/Profiler/Private/ProfilerSession.h
--- a/Engine/Source/Developer/Profiler/Private/ProfilerSession.h
+++ b/Engine/Source/Developer/Profiler/Private/ProfilerSession.h
@@ -392,42 +392,7 @@ protected:
 	 * @param StatMetaData - the stat metadata received from the connected profiler client
 	 *
 	 */
-	void Update( const FStatMetaData& ClientStatMetaData )
-	{
-		PROFILER_SCOPE_LOG_TIME( TEXT( "FProfilerStatMetaData.Update" ), );
-
-		FStatMetaData LocalCopy;
-		{
-			FScopeLock Lock( ClientStatMetaData.CriticalSection );
-			LocalCopy = ClientStatMetaData;
-		}
-
-		// Iterate through all thread descriptions.
-		ThreadDescriptions.Append( LocalCopy.ThreadDescriptions );
-
-		// Initialize fake stat for Self.
-		const uint32 NoGroupID = 0;
-
-		InitializeGroup( NoGroupID, "NoGroup" );
-		InitializeStat( 0, NoGroupID, TEXT( "Self" ), STATTYPE_CycleCounter );
-		InitializeStat( 1, NoGroupID, FStatConstants::NAME_ThreadRoot.GetPlainNameString(), STATTYPE_CycleCounter, FStatConstants::NAME_ThreadRoot );
-
-		// Iterate through all stat group descriptions.
-		for( auto It = LocalCopy.GroupDescriptions.CreateConstIterator(); It; ++It )
-		{
-			const FStatGroupDescription& GroupDesc = It.Value();
-			InitializeGroup( GroupDesc.ID, GroupDesc.Name );
-		}
-
-		// Iterate through all stat descriptions.
-		for( auto It = LocalCopy.StatDescriptions.CreateConstIterator(); It; ++It )
-		{
-			const FStatDescription& StatDesc = It.Value();
-			InitializeStat( StatDesc.ID, StatDesc.GroupID, StatDesc.Name, (EStatType)StatDesc.StatType );
-		}
-
-		SecondsPerCycle = LocalCopy.SecondsPerCycle;
-	}
+	void Update( const FStatMetaData& ClientStatMetaData );
 
 public:
 	/**
@@ -1149,24 +1114,10 @@ protected:
 	);
 
 	/** Called when this profiler session receives a new profiler data. */
-	void UpdateProfilerData( const FProfilerDataFrame& Content )
-	{
-		FrameToProfilerDataMapping.FindOrAdd( Content.Frame ) = Content;
-		FrameToProcess.Add( Content.Frame );
-	}
+	void UpdateProfilerData( const FProfilerDataFrame& Content );
 
 	/** Called when this profiler session receives information that the meta data has been updated. @see FProfilerMetaDataUpdateDelegate and IProfilerClient */
-	void UpdateMetadata( const FStatMetaData& InClientStatMetaData )
-	{
-		const uint32 CurrentStatMetaDataSize = InClientStatMetaData.GetMetaDataSize();
-		if( CurrentStatMetaDataSize != StatMetaDataSize )
-		{
-			ClientStatMetadata = &InClientStatMetaData;
-			bRequestStatMetadataUpdate = true;
-
-			StatMetaDataSize = CurrentStatMetaDataSize;
-		}	
-	}
+	void UpdateMetadata( const FStatMetaData& InClientStatMetaData );
 
 	/**
 	 * Updates aggregated stats.
@@ -1189,7 +1140,7 @@ protected:
 	void EventGraphCombineAndAdd( const FEventGraphDataRef Current, const uint32 NumFrames );
 
 	/** Called when the capture file has been fully loaded. */
-	void LoadComplete() { bHasAllProfilerData = true; }
+	void LoadComplete();
 
 protected:
 	/** All aggregated stats, stored as StatID -> FProfilerAggregatedStat. */
diff --git a/Engine/Source/Developer/ProfilerClient/Private/ProfilerClientManager.cpp b/Engine/Source/Developer/ProfilerClient/Private/ProfilerClientManager.cpp
--- a/Engine/Source/Developer/ProfilerClient/Private/ProfilerClientManager.cpp
+++ b/Engine/Source/Developer/ProfilerClient/Private/ProfilerClientManager.cpp
@@ -1064,8 +1064,6 @@ void FServiceConnection::GenerateCycleGraphs(const FRawStatStackNode& Root, TMap
 
 bool FServiceConnection::ReadAndConvertStatMessages( FArchive& Reader, bool bUseInAsync )
 {
-	uint64 ReadMessages = 0;
-
 	// Buffer used to store the compressed and decompressed data.
 	TArray<uint8> SrcData;
 	TArray<uint8> DestData;
@@ -1120,19 +1118,21 @@ bool FServiceConnection::ReadAndConvertStatMessages( FArchive& Reader, bool bUse
 								}
 							}
 						}
+
+
 					}
 
-					new (Messages)FStatMessage( Message );
+					new (Messages)FStatMessage( Message );			
 				}
 				else
 				{
 					break;
-				}
+				}	
+			}
 
-				if( !bUseInAsync && DataFrames.Num() < FProfilerClientManager::MaxFramesPerTick )
-				{
-					return false;
-				}
+			if( !bUseInAsync && DataFrames.Num() < FProfilerClientManager::MaxFramesPerTick )
+			{
+				return false;
 			}
 		}
 	}
diff --git a/Engine/Source/Developer/ProfilerClient/Private/ProfilerClientManager.h b/Engine/Source/Developer/ProfilerClient/Private/ProfilerClientManager.h
--- a/Engine/Source/Developer/ProfilerClient/Private/ProfilerClientManager.h
+++ b/Engine/Source/Developer/ProfilerClient/Private/ProfilerClientManager.h
@@ -11,6 +11,7 @@ struct FServiceConnection
 {
 	FServiceConnection()
 		: DataLoadingProgress( 0.0f )
+		, ReadMessages(0)
 	{
 		MetaData.CriticalSection = &CriticalSection;
 		CurrentFrame = 0;
@@ -70,6 +71,9 @@ struct FServiceConnection
 	/** Current frame of data to process */
 	int64 CurrentFrame;
 
+	/** Current number of read messages. */
+	uint64 ReadMessages;
+
 	/** Messages received and pending process, they are stored in a map as we can receive them out of order */
 	TMap<int64, TArray<uint8> > PendingMessages;
 
