Stats - Core implementation of the malloc profiler using the stats system (usage: run game with "-memoryprofiler", type in console "stat stopfile" to stop profiling)

diff --git a/Engine/Source/Developer/Profiler/Private/ProfilerCommands.cpp b/Engine/Source/Developer/Profiler/Private/ProfilerCommands.cpp
--- a/Engine/Source/Developer/Profiler/Private/ProfilerCommands.cpp
+++ b/Engine/Source/Developer/Profiler/Private/ProfilerCommands.cpp
@@ -47,7 +47,7 @@ void FProfilerCommands::RegisterCommands()
 
 	UI_COMMAND( OpenSettings, "Settings", "Opens the settings for the profiler", EUserInterfaceActionType::Button, FInputGesture( EModifierKey::Control, EKeys::O ) );
 
-	UI_COMMAND( ProfilerManager_Load, "Load", "Loads profiler data", EUserInterfaceActionType::Button, FInputGesture( EModifierKey::Control, EKeys::L ) );
+	UI_COMMAND( ProfilerManager_Load, "Load", "", EUserInterfaceActionType::Button, FInputGesture(/* EModifierKey::Control, EKeys::L */) );
 	UI_COMMAND( ProfilerManager_ToggleLivePreview, "Live preview", "Toggles the real time live preview", EUserInterfaceActionType::ToggleButton, FInputGesture() );
 
 	UI_COMMAND( DataGraph_ToggleViewMode, "Toggle graph view mode", "Toggles the data graph view mode between time based and index based", EUserInterfaceActionType::Button, FInputGesture() );
@@ -413,8 +413,6 @@ void FProfilerActionManager::ToggleDataCapture_Execute( const FGuid SessionInsta
 
 bool FProfilerActionManager::ToggleDataCapture_CanExecute( const FGuid SessionInstanceID ) const
 {
-	// Disabled capture button because current implementation in the profiler service is not thread-safe.
-#if	0
 	// One session instance
 	if( SessionInstanceID.IsValid() )
 	{
@@ -427,7 +425,6 @@ bool FProfilerActionManager::ToggleDataCapture_CanExecute( const FGuid SessionIn
 		const bool bCanExecute = This->ActiveSession.IsValid() && This->ProfilerType == EProfilerSessionTypes::Live && This->GetProfilerInstancesNum() > 0;
 		return bCanExecute;
 	}
-#endif // 0
 
 	return false;
 }
diff --git a/Engine/Source/Developer/Profiler/Private/ProfilerSession.cpp b/Engine/Source/Developer/Profiler/Private/ProfilerSession.cpp
--- a/Engine/Source/Developer/Profiler/Private/ProfilerSession.cpp
+++ b/Engine/Source/Developer/Profiler/Private/ProfilerSession.cpp
@@ -522,6 +522,80 @@ bool FRawProfilerSession::HandleTicker( float DeltaTime )
 	return true;
 }
 
+static double GetSecondsPerCycle( const FStatPacketArray& Frame )
+{
+	const FName SecondsPerCycleFName = FName(TEXT("//STATGROUP_Engine//STAT_SecondsPerCycle///Seconds$32$Per$32$Cycle///////STATCAT_Advanced////"));
+	double Result = 0;
+
+	for( int32 PacketIndex = 0; PacketIndex < Frame.Packets.Num(); PacketIndex++ )
+	{
+		const FStatPacket& Packet = *Frame.Packets[PacketIndex];
+		const FStatMessagesArray& Data = Packet.StatMessages;
+
+		if( Packet.ThreadType == EThreadType::Game )
+		{
+			for( int32 Index = 0; Index < Data.Num(); Index++ )
+			{
+				const FStatMessage& Item = Data[Index];
+				check( Item.NameAndInfo.GetFlag( EStatMetaFlags::DummyAlwaysOne ) ); 
+
+				const FName LongName = Item.NameAndInfo.GetEncodedName();
+				if( LongName.IsEqual( SecondsPerCycleFName, ENameCase::IgnoreCase, false ) )
+				{
+					Result = Item.GetValue_double();
+					UE_LOG( LogStats, Log, TEXT( "STAT_SecondsPerCycle is %f [ns]" ), Result*1000*1000 );
+
+					PacketIndex = Frame.Packets.Num();
+					break;
+				}
+			}	
+		}
+			
+	}
+	return Result;
+}
+
+static int64 GetFastThreadFrameTimeInternal( const FStatPacketArray& Frame, EThreadType::Type ThreadType )
+{
+	int64 Result = 0;
+
+	for( int32 PacketIndex = 0; PacketIndex < Frame.Packets.Num(); PacketIndex++ )
+	{
+		const FStatPacket& Packet = *Frame.Packets[PacketIndex];
+
+		if( Packet.ThreadType == ThreadType )
+		{
+			const FStatMessagesArray& Data = Packet.StatMessages;
+			for (int32 Index = 0; Index < Data.Num(); Index++)
+			{
+				FStatMessage const& Item = Data[Index];
+				EStatOperation::Type Op = Item.NameAndInfo.GetField<EStatOperation>();
+				FName LongName = Item.NameAndInfo.GetRawName();
+				if (Op == EStatOperation::CycleScopeStart)
+				{
+					check(Item.NameAndInfo.GetFlag(EStatMetaFlags::IsCycle));
+					Result -= Item.GetValue_int64();
+					break;
+				}
+			}
+			for (int32 Index = Data.Num() - 1; Index >= 0; Index--)
+			{
+				FStatMessage const& Item = Data[Index];
+				EStatOperation::Type Op = Item.NameAndInfo.GetField<EStatOperation>();
+				FName LongName = Item.NameAndInfo.GetRawName();
+				if (Op == EStatOperation::CycleScopeEnd)
+				{
+
+					check(Item.NameAndInfo.GetFlag(EStatMetaFlags::IsCycle));
+					Result += Item.GetValue_int64();
+					break;
+				}
+			}
+		}
+	}
+	return Result;
+}
+
 void FRawProfilerSession::PrepareLoading()
 {
 	SCOPE_LOG_TIME_FUNC();
@@ -548,6 +622,7 @@ void FRawProfilerSession::PrepareLoading()
 
 	const bool bIsFinalized = Stream.Header.IsFinalized();
 	check( bIsFinalized );
+	check( Stream.Header.Version == EStatMagicWithHeader::VERSION_5 );
 	StatsThreadStats.MarkAsLoaded();
 
 	TArray<FStatMessage> Messages;
@@ -558,6 +633,11 @@ void FRawProfilerSession::PrepareLoading()
 		Stream.ReadFNamesAndMetadataMessages( *FileReader, MetadataMessages );
 		StatsThreadStats.ProcessMetaDataOnly( MetadataMessages );
 
+		const FName F00245 = FName(245, 245, 0);
+		
+		const FName F11602 = FName(11602, 11602, 0);
+		const FName F06394 = FName(6394, 6394, 0);
+
 		const int64 CurrentFilePos = FileReader->Tell();
 
 		// Update profiler's metadata.
@@ -571,101 +651,133 @@ void FRawProfilerSession::PrepareLoading()
 		TArray<uint8> SrcArray;
 		TArray<uint8> DestArray;
 		const bool bHasCompressedData = Stream.Header.HasCompressedData();
-
-		// Temporary fix for setting the correct value for the SecondsPerCycle.
+		check(bHasCompressedData);
+
+		TMap<int64, FStatPacketArray> CombinedHistory;
+		int64 TotalPacketSize = 0;
+		int64 MaximumPacketSize = 0;
+		// Read all packets sequentially, force by the memory profiler which is now a part of the raw stats.
+		// !!CAUTION!! Frame number in the raw stats is pointless, because it is time based, not frame based.
+		// Background threads usually execute time consuming operations, so the frame number won't be valid.
+		// Needs to be combined by the thread and the time, not by the frame number.
 		{
-			SCOPE_LOG_TIME( TEXT( "Looking for the STAT_SecondsPerCycle" ), nullptr );
-			UE_LOG( LogStats, Log, TEXT( "Looking for the STAT_SecondsPerCycle" ) );
-			FileReader->Seek( Stream.FramesInfo[0].FrameFileOffset );
+			int64 FrameOffset0 = Stream.FramesInfo[0].FrameFileOffset;
+			FileReader->Seek( FrameOffset0 );
 
-			if( bHasCompressedData )
+			const int64 FileSize = FileReader->TotalSize();
+
+			while( FileReader->Tell() < FileSize )
 			{
 				// Read the compressed data.
 				FCompressedStatsData UncompressedData( SrcArray, DestArray );
 				*FileReader << UncompressedData;
-			}
+				if( UncompressedData.HasReachedEndOfCompressedData() )
+				{
+					break;
+				}
 
-			// Select the proper archive.
-			FMemoryReader MemoryReader( DestArray, true );
-			FArchive& Archive = bHasCompressedData ? MemoryReader : *FileReader;
+				FMemoryReader MemoryReader( DestArray, true );
 
-			FStatPacketArray StatPacketArray;
+				FStatPacket* StatPacket = new FStatPacket();
+				Stream.ReadStatPacket( MemoryReader, *StatPacket );
+				
+				const int64 FrameNum = StatPacket->Frame;
+				FStatPacketArray& Frame = CombinedHistory.FindOrAdd(FrameNum);
+			
+				// Check if we need to combine packets from the same thread.
+				FStatPacket** CombinedPacket = Frame.Packets.FindByPredicate([&](FStatPacket* Item) -> bool
+				{
+					return Item->ThreadId == StatPacket->ThreadId;
+				});
+				
+				if( CombinedPacket )
+				{
+					(*CombinedPacket)->StatMessages += StatPacket->StatMessages;
+				}
+				else
+				{
+					Frame.Packets.Add(StatPacket);
+				}
 
-			int64 TargetFrame;
-			Archive << TargetFrame;
+				const int64 CurrentPos = FileReader->Tell();
+				const int32 PctPos = int32(100.0f*CurrentPos/FileSize);
 
-			int32 NumPackets;
-			Archive << NumPackets;
+				UE_LOG( LogStats, Log, TEXT( "%3i Processing FStatPacket: Frame %5i for thread %5i with %6i messages (%.1f MB)" ), 
+					PctPos, 
+					StatPacket->Frame, 
+					StatPacket->ThreadId, 
+					StatPacket->StatMessages.Num(), 
+					StatPacket->StatMessages.GetAllocatedSize()/1024.0f/1024.0f );
 
-			for( int32 PacketIndex = 0; PacketIndex < NumPackets; PacketIndex++ )
-			{
-				FStatPacket* ToRead = new FStatPacket();
-				Stream.ReadStatPacket( Archive, *ToRead, bIsFinalized );
-				StatPacketArray.Packets.Add( ToRead );
+				const int64 PacketSize = StatPacket->StatMessages.GetAllocatedSize();
+				TotalPacketSize += PacketSize;
+				MaximumPacketSize = FMath::Max( MaximumPacketSize, PacketSize );
 			}
+		}
+
+		UE_LOG( LogStats, Log, TEXT( "TotalPacketSize: %.1f MB, Max: %1f MB" ), 
+			TotalPacketSize/1024.0f/1024.0f, 
+			MaximumPacketSize/1024.0f/1024.0f );
 
-			for( int32 PacketIndex = 0; PacketIndex < StatPacketArray.Packets.Num(); PacketIndex++ )
+		TArray<int64> Frames;
+		CombinedHistory.GenerateKeyArray(Frames);
+		Frames.Sort();
+		const int64 MiddleFrame = Frames[Frames.Num()/2];
+
+		const bool bProcessMemory = false;
+		if( bProcessMemory )
+		{
+			ProcessMemoryOperations(CombinedHistory);
+		}
+		else
+		{
+			// Remove all frames without the game thread messages.
+			for( int32 FrameIndex = 0; FrameIndex < Frames.Num(); ++FrameIndex )
 			{
-				const FStatPacket& StatPacket = *StatPacketArray.Packets[PacketIndex];
-				const FStatMessagesArray& Data = StatPacket.StatMessages;
+				const int64 TargetFrame = Frames[FrameIndex];
+				const FStatPacketArray& Frame = CombinedHistory.FindChecked(TargetFrame);
 
-				for( int32 Index = 0; Index < Data.Num(); Index++ )
-				{
-					const FStatMessage& Item = Data[Index];
-					check( Item.NameAndInfo.GetFlag( EStatMetaFlags::DummyAlwaysOne ) );  // we should never be sending short names to the stats anymore
+				const double GameThreadTimeMS = GetMetaData()->ConvertCyclesToMS( GetFastThreadFrameTimeInternal(Frame,EThreadType::Game) );
 
-					const FName ShortName = Item.NameAndInfo.GetShortName();
-					if( ShortName == TEXT( "STAT_SecondsPerCycle" ) )
-					{
-						StatMetaData->SecondsPerCycle = Item.GetValue_double();
-						UE_LOG( LogStats, Log, TEXT( "STAT_SecondsPerCycle is %f [ns]" ), StatMetaData->GetSecondsPerCycle()*1000*1000 );
-						
-						PacketIndex = StatPacketArray.Packets.Num();
-						break;
-					}
+				if( GameThreadTimeMS == 0.0f )
+				{
+					CombinedHistory.Remove(TargetFrame);
+					Frames.RemoveAt(FrameIndex);
+					FrameIndex--;
 				}
 			}
 		}
-
+	
+		StatMetaData->SecondsPerCycle = GetSecondsPerCycle( CombinedHistory.FindChecked(MiddleFrame) );
 		check( StatMetaData->GetSecondsPerCycle() > 0.0 );
 
+		//const int32 FirstGameThreadFrame = FindFirstFrameWithGameThread( CombinedHistory, Frames );
+
+		// Prepare profiler frame.
 		{
 			SCOPE_LOG_TIME( TEXT( "Preparing profiler frames" ), nullptr );
 
 			// Prepare profiler frames.
 			double ElapsedTimeMS = 0;
-			for( int32 FrameIndex = 0; FrameIndex < Stream.FramesInfo.Num(); ++FrameIndex )
+
+			for( int32 FrameIndex = 0; FrameIndex < Frames.Num(); ++FrameIndex )
 			{
-				const FStatsFrameInfo& StatsFrameInfo = Stream.FramesInfo[FrameIndex];
-				FileReader->Seek( StatsFrameInfo.FrameFileOffset );
+				const int64 TargetFrame = Frames[FrameIndex];
+				const FStatPacketArray& Frame = CombinedHistory.FindChecked(TargetFrame);
+
+				const double GameThreadTimeMS = GetMetaData()->ConvertCyclesToMS( GetFastThreadFrameTimeInternal(Frame,EThreadType::Game) );
 
-				if( bHasCompressedData )
+				if( GameThreadTimeMS == 0.0f )
 				{
-					// Read the compressed data.
-					FCompressedStatsData UncompressedData( SrcArray, DestArray );
-					*FileReader << UncompressedData;
+					continue;
 				}
 
-				// Select the proper archive.
-				FMemoryReader MemoryReader( DestArray, true );
-				FArchive& Archive = bHasCompressedData ? MemoryReader : *FileReader;
-
-				int64 TargetFrame;
-				Archive << TargetFrame;
-
-				const double GameThreadTimeMS = GetMetaData()->ConvertCyclesToMS( (uint32)StatsFrameInfo.ThreadCycles.FindChecked( GameThreadID ) );
+				const double RenderThreadTimeMS = GetMetaData()->ConvertCyclesToMS( GetFastThreadFrameTimeInternal(Frame,EThreadType::Renderer) );
 
 				// Update mini-view, convert from cycles to ms.
 				TMap<uint32, float> ThreadTimesMS;
-				for( auto InnerIt = StatsFrameInfo.ThreadCycles.CreateConstIterator(); InnerIt; ++InnerIt )
-				{
-					const uint32 ThreadID = InnerIt.Key();
-					const bool bCanBeAdded = GameThreadID == ThreadID || GetMetaData()->GetRenderThreadID().Contains( ThreadID );
-					if( bCanBeAdded )
-					{
-						ThreadTimesMS.Add( ThreadID, GetMetaData()->ConvertCyclesToMS( InnerIt.Value() ) );
-					}
-				}
+				ThreadTimesMS.Add( GameThreadID, GameThreadTimeMS );
+				ThreadTimesMS.Add( GetMetaData()->GetRenderThreadID()[0], RenderThreadTimeMS );
 
 				// Pass the reference to the stats' metadata.
 				OnAddThreadTime.ExecuteIfBound( FrameIndex, ThreadTimesMS, StatMetaData );
@@ -675,55 +787,26 @@ void FRawProfilerSession::PrepareLoading()
 				FProfilerFrame* ProfilerFrame = new FProfilerFrame( TargetFrame, GameThreadTimeMS, ElapsedTimeMS );
 				ProfilerFrame->ThreadTimesMS = ThreadTimesMS;
 				ProfilerStream.AddProfilerFrame( TargetFrame, ProfilerFrame );
-
-				// Thread cycles can be used to generate low-res thread view.
-				// @TODO yrx 2014-04-23 
 			}
-			int64 FrameOffset0 = Stream.FramesInfo[0].FrameFileOffset;
-			FileReader->Seek( FrameOffset0 );
 		}
-
-#if	1
+	
+		// Process the raw stats data.
 		{
-			SCOPE_LOG_TIME( TEXT( "Processing raw stats frames" ), nullptr );
+			SCOPE_LOG_TIME( TEXT( "Processing the raw stats" ), nullptr );
 
 			double CycleCounterAdjustmentMS = 0.0f;
 
 			// Read the raw stats messages.
-			for( int32 FrameIndex = 0; FrameIndex < Stream.FramesInfo.Num(); ++FrameIndex )
+			for( int32 FrameIndex = 0; FrameIndex < Frames.Num()-1; ++FrameIndex )
 			{
-				FProfilerFrame* ProfilerFrame = ProfilerStream.GetProfilerFrame( FrameIndex );
-				FileReader->Seek( Stream.FramesInfo[FrameIndex].FrameFileOffset );
-
-				UE_CLOG( FrameIndex % 8 == 0, LogStats, Log, TEXT( "Processing raw stats frame: %4i/%4i" ), FrameIndex, Stream.FramesInfo.Num() );
-
-				if( bHasCompressedData )
-				{
-					// Read the compressed data.
-					FCompressedStatsData UncompressedData( SrcArray, DestArray );
-					*FileReader << UncompressedData;
-				}
+				const int64 TargetFrame = Frames[FrameIndex];
+				const FStatPacketArray& Frame = CombinedHistory.FindChecked(TargetFrame);
 
-				// Select the proper archive.
-				FMemoryReader MemoryReader( DestArray, true );
-				FArchive& Archive = bHasCompressedData ? MemoryReader : *FileReader;
-
-				FStatPacketArray Packet;
+				FProfilerFrame* ProfilerFrame = ProfilerStream.GetProfilerFrame( FrameIndex );
 
-				int64 TargetFrame;
-				Archive << TargetFrame;
+				UE_CLOG( FrameIndex % 8 == 0, LogStats, Log, TEXT( "Processing raw stats frame: %4i/%4i" ), FrameIndex, Frames.Num() );
 
-				int32 NumPackets;
-				Archive << NumPackets;
-
-				for( int32 PacketIndex = 0; PacketIndex < NumPackets; PacketIndex++ )
-				{
-					FStatPacket* ToRead = new FStatPacket();
-					Stream.ReadStatPacket( Archive, *ToRead, bIsFinalized );
-					Packet.Packets.Add( ToRead );
-				}
-
-				ProcessStatPacketArray( Packet, *ProfilerFrame, FrameIndex ); // or ProfilerFrame->TargetFrame
+				ProcessStatPacketArray( Frame, *ProfilerFrame, FrameIndex ); // or ProfilerFrame->TargetFrame
 
 				// Find the first cycle counter for the game thread.
 				if( CycleCounterAdjustmentMS == 0.0f )
@@ -731,10 +814,6 @@ void FRawProfilerSession::PrepareLoading()
 					CycleCounterAdjustmentMS = ProfilerFrame->Root->CycleCounterStartTimeMS;
 				}
 
-				TMap<uint32, int64> ThreadCycles;
-				// Serialize thread cycles.
-				Archive << ThreadCycles;
-
 				// Update thread time and mark profiler frame as valid and ready for use.
 				ProfilerFrame->MarkAsValid();
 			}
@@ -742,8 +821,6 @@ void FRawProfilerSession::PrepareLoading()
 			// Adjust all profiler frames.
 			ProfilerStream.AdjustCycleCounters( CycleCounterAdjustmentMS );
 		}
-
-#endif // 0
 	}
 
 	const int64 AllocatedSize = ProfilerStream.GetAllocatedSize();
@@ -951,6 +1028,10 @@ void FRawProfilerSession::ProcessStatPacketArray( const FStatPacketArray& StatPa
 		{
 			ThreadFName = NAME_RenderThread;
 		}
+		else if( StatPacket.ThreadType == EThreadType::Game )
+		{
+			GameThreadFName = ThreadFName;
+		}
 
 		FProfilerStackNode* ThreadNode = ThreadNodes.FindRef( ThreadFName );
 		if( !ThreadNode )
@@ -976,10 +1057,6 @@ void FRawProfilerSession::ProcessStatPacketArray( const FStatPacketArray& StatPa
 			ThreadNode = ThreadNodes.Add( ThreadFName, new FProfilerStackNode( nullptr, ThreadMessage, ThreadRootSampleIndex, FrameIndex ) );
 		}
 
-		if( StatPacket.ThreadType == EThreadType::Game )
-		{
-			GameThreadFName = ThreadFName;
-		}
 
 		TArray<const FStatMessage*> StartStack;
 		TArray<FProfilerStackNode*> Stack;
@@ -990,15 +1067,22 @@ void FRawProfilerSession::ProcessStatPacketArray( const FStatPacketArray& StatPa
 		for( int32 Index = 0; Index < Data.Num(); Index++ )
 		{
 			const FStatMessage& Item = Data[Index];
-			check( Item.NameAndInfo.GetFlag( EStatMetaFlags::DummyAlwaysOne ) );  // we should never be sending short names to the stats anymore
 
 			const EStatOperation::Type Op = Item.NameAndInfo.GetField<EStatOperation>();
 			const FName LongName = Item.NameAndInfo.GetRawName();
 			const FName ShortName = Item.NameAndInfo.GetShortName();
 
-			if( Op == EStatOperation::CycleScopeStart || Op == EStatOperation::CycleScopeEnd )
+			const FName RenderingThreadTickCommandName = TEXT("RenderingThreadTickCommand");
+
+			// Workaround for render thread hierarchy. EStatOperation::AdvanceFrameEventRenderThread is called within the scope.
+			if( ShortName == RenderingThreadTickCommandName )
+			{
+				continue;
+			}
+
+			if( Op == EStatOperation::CycleScopeStart || Op == EStatOperation::CycleScopeEnd || Op == EStatOperation::AdvanceFrameEventRenderThread )
 			{
-				check( Item.NameAndInfo.GetFlag( EStatMetaFlags::IsCycle ) );
+				//check( Item.NameAndInfo.GetFlag( EStatMetaFlags::IsCycle ) );
 				if( Op == EStatOperation::CycleScopeStart )
 				{
 					FProfilerStackNode* ChildNode = new FProfilerStackNode( Current, Item, -1, FrameIndex );
@@ -1021,15 +1105,26 @@ void FRawProfilerSession::ProcessStatPacketArray( const FStatPacketArray& StatPa
 					StartStack.Add( &Item );
 					Current = ChildNode;
 				}
+				// Workaround for render thread hierarchy. EStatOperation::AdvanceFrameEventRenderThread is called within the scope.
+				if( Op == EStatOperation::AdvanceFrameEventRenderThread )
+				{
+					int k=0;k++;
+				}
 				if( Op == EStatOperation::CycleScopeEnd )
 				{
 					const FStatMessage ScopeStart = *StartStack.Pop();
 					const FStatMessage ScopeEnd = Item;
 					const int64 Delta = int32( uint32( ScopeEnd.GetValue_int64() ) - uint32( ScopeStart.GetValue_int64() ) );
-					Current->CyclesEnd = Current->CyclesStart + int64( Delta );
+					Current->CyclesEnd = Current->CyclesStart + Delta;
 
 					Current->CycleCounterStartTimeMS = MetaData->ConvertCyclesToMS( Current->CyclesStart );
 					Current->CycleCounterEndTimeMS = MetaData->ConvertCyclesToMS( Current->CyclesEnd );
+
+					if( Current->CycleCounterStartTimeMS > Current->CycleCounterEndTimeMS )
+					{
+						int k=0;k++;
+					}
+
 					check( Current->CycleCounterEndTimeMS >= Current->CycleCounterStartTimeMS );
 
 					FProfilerStackNode* ChildNode = Current;
@@ -1090,3 +1185,563 @@ void FRawProfilerSession::ProcessStatPacketArray( const FStatPacketArray& StatPa
 
 	out_ProfilerFrame.SortChildren();
 }
+
+
+//-----------------------------------------------------------------------------
+
+const TCHAR* CallstackSeparator = TEXT("+");
+
+static FString EncodeCallstack( const TArray<FName>& Callstack )
+{
+	FString Result;
+	for( const auto& Name : Callstack )
+	{
+		Result += TTypeToString<int32>::ToString((int32)Name.GetComparisonIndex());
+		Result += CallstackSeparator;
+	}
+	return Result;
+}
+
+static void DecodeCallstack( const FName& EncodedCallstack, TArray<FString>& out_DecodedCallstack )
+{
+	EncodedCallstack.ToString().ParseIntoArray( &out_DecodedCallstack, CallstackSeparator, true );
+}
+
+/** Simple allocation info. Assumes the sequence tag are unique and doesn't wrap. */
+struct FAllocationInfo
+{
+	uint64 Ptr;
+	int64 Size;
+	uint32 SequenceTag;
+	FName EncodedCallstack;
+	EMemoryOperation Op;
+	bool bHasBrokenCallstack;
+	
+	FAllocationInfo(
+		uint64 InPtr, 
+		int64 InSize, 
+		const TArray<FName>& InCallstack, 
+		uint32 InSequenceTag, 
+		EMemoryOperation InOp, 
+		bool bInHasBrokenCallstack)
+		: Ptr(InPtr)
+		, Size(InSize)
+		, SequenceTag(InSequenceTag)
+		, EncodedCallstack(*EncodeCallstack(InCallstack))
+		, Op(InOp)
+		, bHasBrokenCallstack(bInHasBrokenCallstack)
+	{
+	}
+
+	FAllocationInfo(const FAllocationInfo& Other)
+		: Ptr(Other.Ptr)
+		, Size(Other.Size)
+		, SequenceTag(Other.SequenceTag)
+		, EncodedCallstack(Other.EncodedCallstack)
+		, Op(Other.Op)
+		, bHasBrokenCallstack(Other.bHasBrokenCallstack)
+	{}
+
+	FAllocationInfo()
+		: Ptr(0)
+		, Size(0)
+		, SequenceTag(0)
+		, Op(EMemoryOperation::Invalid)
+	{}
+
+
+	bool operator<( const FAllocationInfo& Other ) const
+	{
+		return SequenceTag < Other.SequenceTag;
+	}
+
+	bool operator==( const FAllocationInfo& Other ) const
+	{
+		return SequenceTag == Other.SequenceTag;
+	}
+};
+
+struct FAllocationInfoGreater
+{
+	FORCEINLINE bool operator()(const FAllocationInfo& A, const FAllocationInfo& B) const { return B.Size < A.Size; }
+};
+
+
+struct FSizeAndCount
+{
+	uint64 Size;
+	uint64 Count;
+	FSizeAndCount()
+		: Size(0)
+		, Count(0)
+	{}
+};
+
+struct FSizeAndCountGreater
+{
+	FORCEINLINE bool operator()(const FSizeAndCount& A, const FSizeAndCount& B) const { return B.Size < A.Size; }
+};
+
+/** Holds stats stack state, used to preserve continuity when the game frame has changed. */
+struct FStackState
+{
+	FStackState()
+		: bIsBrokenCallstack(false)
+	{}
+
+	/** Call stack. */
+	TArray<FName> Stack;
+
+	/** Current function name. */
+	FName Current;
+
+	/** Whether this callstack is marked as broken due to mismatched start and end scope cycles. */
+	bool bIsBrokenCallstack;
+};
+
+FString GetCallstack( const FName& EncodedCallstack )
+{
+	FString Result;
+
+	TArray<FString> DecodedCallstack;
+	DecodeCallstack( EncodedCallstack, DecodedCallstack );
+
+	for( int32 Index = DecodedCallstack.Num() - 1; Index >= 0; --Index )
+	{
+		NAME_INDEX NameIndex = 0;
+		TTypeFromString<NAME_INDEX>::FromString( NameIndex, *DecodedCallstack[Index] );
+
+		const FName LongName = FName(NameIndex, NameIndex, 0);
+
+		const FString ShortName = FStatNameAndInfo::GetShortNameFrom(LongName).ToString();
+		const FString Group = FStatNameAndInfo::GetGroupNameFrom(LongName).ToString();
+		FString Desc = FStatNameAndInfo::GetDescriptionFrom(LongName);
+		Desc.Trim();
+
+		const bool bHasDesc = Desc.Len() > 0;
+		if( Desc != ShortName )
+		{
+			if( bHasDesc )
+			{
+				Desc += TEXT( " [" );
+			}
+			Desc += ShortName;
+			if( bHasDesc )
+			{
+				Desc += TEXT( "]" );
+			}
+		}
+
+		//Result += FString::Printf( TEXT("%s [%s] -> "), *Desc, *Group );
+		Result += Desc;
+		if( Index > 0 )
+		{
+			Result += TEXT(" <- ");
+		}
+	}
+
+	return Result;
+}
+
+// New version.
+void FRawProfilerSession::ProcessMemoryOperations( const TMap<int64, FStatPacketArray>& CombinedHistory )
+{
+	// This is only example code, no fully implemented, may sometimes crash.
+	// This code is not optimized. 
+
+	uint64 NumMemoryOperations = 0;
+
+	// Generate frames
+	TArray<int64> Frames;
+	CombinedHistory.GenerateKeyArray(Frames);
+	Frames.Sort();
+
+	const FProfilerStatMetaDataRef MetaData = GetMetaData();
+
+	// Raw stats callstack for this stat packet array.
+	TMap<FName,FStackState> StackStates;
+
+	// All allocation ordered by the sequence tag.
+	// There is an assumption that the sequence tag will not turn-around.
+	//TMap<uint32, FAllocationInfo> SequenceAllocationMap;
+	TArray<FAllocationInfo> SequenceAllocationArray;
+
+	// Pass 1.
+	// Read all stats messages, parse all memory operations and decode callstacks.
+	const int64 FirstFrame = 0;
+	for( int32 FrameIndex = 0; FrameIndex < Frames.Num(); ++FrameIndex )
+	{
+		UE_LOG( LogStats, Log, TEXT( "Processing %i/%i frame" ), FrameIndex, Frames.Num() );
+
+		const int64 TargetFrame = Frames[FrameIndex];
+		const int64 Diff = TargetFrame-FirstFrame;
+		const FStatPacketArray& Frame = CombinedHistory.FindChecked(TargetFrame);
+
+		for( int32 PacketIndex = 0; PacketIndex < Frame.Packets.Num(); PacketIndex++ )
+		{
+			UE_LOG( LogStats, Log, TEXT( "Processing packet %i/%i" ), PacketIndex, Frame.Packets.Num() );
+
+			const FStatPacket& StatPacket = *Frame.Packets[PacketIndex];
+			const FName& ThreadFName = StatsThreadStats.Threads.FindChecked( StatPacket.ThreadId );
+			const uint32 NewThreadID = MetaData->ThreadIDtoStatID.FindChecked( StatPacket.ThreadId );
+
+			FStackState* StackState = StackStates.Find(ThreadFName);
+			if( !StackState )
+			{
+				StackState = &StackStates.Add(ThreadFName);
+				StackState->Stack.Add( ThreadFName );
+				StackState->Current = ThreadFName;
+			}
+
+			const FStatMessagesArray& Data = StatPacket.StatMessages;
+
+			int32 LastPct = 0;
+			for( int32 Index = 0; Index < Data.Num(); Index++ )
+			{
+				const int32 CurrentPct = int32(100.0 * Index / Data.Num());
+				if( CurrentPct > LastPct+10 )
+				{
+					UE_LOG( LogStats, Log, TEXT( "Processing %3i %% (%i/%i) stat messages" ), CurrentPct, Index, Data.Num() );
+					LastPct = CurrentPct;
+				}
+				
+				const FStatMessage& Item = Data[Index];
+
+				const EStatOperation::Type Op = Item.NameAndInfo.GetField<EStatOperation>();
+				const FName LongName = Item.NameAndInfo.GetEncodedName();
+				const FName RawName = Item.NameAndInfo.GetRawName();
+
+				if( Op == EStatOperation::CycleScopeStart || Op == EStatOperation::CycleScopeEnd || Op == EStatOperation::Memory )
+				{
+					if( Op == EStatOperation::CycleScopeStart )
+					{
+						StackState->Stack.Add( RawName );
+						StackState->Current = RawName;
+					}
+					else if( Op == EStatOperation::Memory )
+					{
+						// Experimental code used only to test the implementation.
+						// First memory operation is Alloc or Free
+						const uint64 EncodedPtr = Item.GetValue_Ptr();
+						const bool bIsAlloc = (EncodedPtr & (uint64)EMemoryOperation::Alloc) != 0;
+						const bool bIsFree = (EncodedPtr & (uint64)EMemoryOperation::Free) != 0;
+						const uint64 Ptr = EncodedPtr & ~(uint64)EMemoryOperation::Mask;
+						if( bIsAlloc )
+						{
+							NumMemoryOperations++;
+							// @see FStatsMallocProfilerProxy::TrackAlloc
+							// After alloc ptr message there is always alloc size message and the sequence tag.
+							Index++;
+							const FStatMessage& AllocSizeMessage = Data[Index];
+							const int64 AllocSize = AllocSizeMessage.GetValue_int64();
+
+							// Read operation sequence tag.
+							Index++;
+							const FStatMessage& SequenceTagMessage = Data[Index];
+							const uint32 SequenceTag = SequenceTagMessage.GetValue_int64();
+
+							// Create a callstack.
+							TArray<FName> StatsBasedCallstack;
+							for( const auto& RawName : StackState->Stack )
+							{
+								StatsBasedCallstack.Add(RawName);
+							}
+
+							// Add a new allocation.
+							SequenceAllocationArray.Add( 
+								FAllocationInfo(
+									Ptr,
+									AllocSize, 
+									StatsBasedCallstack, 
+									SequenceTag, 
+									EMemoryOperation::Alloc, 						
+									StackState->bIsBrokenCallstack) );
+						}
+						else if( bIsFree )
+						{
+							NumMemoryOperations++;
+							// Read operation sequence tag.
+							Index++;
+							const FStatMessage& SequenceTagMessage = Data[Index];
+							const uint32 SequenceTag = SequenceTagMessage.GetValue_int64();
+
+							// Create a callstack.
+							TArray<FName> StatsBasedCallstack;
+							for( const auto& RawName : StackState->Stack )
+							{
+								StatsBasedCallstack.Add(RawName);
+							}
+
+							// Add a new free.
+							SequenceAllocationArray.Add( 
+								FAllocationInfo(
+									Ptr,
+									0, 
+									StatsBasedCallstack, 
+									SequenceTag,
+									EMemoryOperation::Free, 					
+									StackState->bIsBrokenCallstack) );
+						}
+						else
+						{
+							UE_LOG(LogStats, Warning, TEXT("Pointer from a memory operation is invalid") );
+						}
+					}
+					else if( Op == EStatOperation::CycleScopeEnd )
+					{
+						if( StackState->Stack.Num() > 1 )
+						{
+							const FName ScopeStart = StackState->Stack.Pop();
+							const FName ScopeEnd = Item.NameAndInfo.GetRawName();
+
+							check(ScopeStart==ScopeEnd);
+
+							StackState->Current = StackState->Stack.Last();
+
+							// The stack should be ok, but it may be partially broken.
+							// This will happen if memory profiling starts in the middle of executing a background thread.
+							StackState->bIsBrokenCallstack = false;
+						}
+						else
+						{
+							const FName ShortName = Item.NameAndInfo.GetShortName();
+
+							UE_LOG(LogStats, Warning, TEXT("Broken cycle scope end %s/%s, current %s"), 
+								*ThreadFName.ToString(), 
+								*ShortName.ToString(),
+								*StackState->Current.ToString() );
+
+							// The stack is completely broken, only has the thread name and the last cycle scope.
+							// Rollback to the thread node.
+							StackState->bIsBrokenCallstack = true;
+							StackState->Stack.Empty();
+							StackState->Stack.Add( ThreadFName );
+							StackState->Current = ThreadFName;
+						}
+					}
+				}
+			}
+		}
+	}
+
+	UE_LOG(LogStats, Warning, TEXT("NumMemoryOperations %llu"), NumMemoryOperations );
+	UE_LOG(LogStats, Warning, TEXT("SequenceAllocationArray.Num %i"), SequenceAllocationArray.Num() );
+
+	// Pass 2.
+	/*
+	TMap<uint32,FAllocationInfo> UniqueSeq;
+	TMultiMap<uint32,FAllocationInfo> OriginalAllocs;
+	TMultiMap<uint32,FAllocationInfo> BrokenAllocs;
+	for( const FAllocationInfo& Alloc : SequenceAllocationArray )
+	{
+		const FAllocationInfo* Found = UniqueSeq.Find(Alloc.SequenceTag);
+		if( !Found )
+		{
+			UniqueSeq.Add(Alloc.SequenceTag,Alloc);
+		}
+		else
+		{
+			OriginalAllocs.Add(Alloc.SequenceTag, *Found);
+			BrokenAllocs.Add(Alloc.SequenceTag, Alloc);
+		}		
+	}
+	*/
+
+	// Sort all memory operation by the sequence tag, iterate through all operation and generate memory usage.
+	SequenceAllocationArray.Sort( TLess<FAllocationInfo>() );
+
+	// Alive allocations.
+	TMap<uint64, FAllocationInfo> AllocationMap;
+	TMultiMap<uint64, FAllocationInfo> FreeWithoutAllocMap;
+	TMultiMap<uint64, FAllocationInfo> DuplicatedAllocMap;
+	int32 NumDuplicatedMemoryOperations = 0;
+	int32 NumFWAMemoryOperations = 0; // FreeWithoutAlloc
+
+	for( const FAllocationInfo& Alloc : SequenceAllocationArray )
+	{
+		const EMemoryOperation MemOp = Alloc.Op;
+		const uint64 Ptr = Alloc.Ptr;
+		const int64 Size = Alloc.Size;
+		const uint32 SequenceTag = Alloc.SequenceTag;
+
+		if( MemOp == EMemoryOperation::Alloc )
+		{
+			const FAllocationInfo* Found = AllocationMap.Find(Ptr);
+
+			if( !Found )
+			{
+				AllocationMap.Add(Ptr, Alloc);
+			}
+			else
+			{
+				const FAllocationInfo* FoundAndFreed = FreeWithoutAllocMap.Find(Found->Ptr);
+				const FAllocationInfo* FoundAndAllocated = FreeWithoutAllocMap.Find(Alloc.Ptr);
+
+				if( FoundAndFreed )
+				{
+					const FString FoundAndFreedCallstack = GetCallstack(FoundAndFreed->EncodedCallstack);
+				}
+
+				if( FoundAndAllocated )
+				{
+					const FString FoundAndAllocatedCallstack = GetCallstack(FoundAndAllocated->EncodedCallstack);
+				}
+
+				NumDuplicatedMemoryOperations++;
+
+				const FString FoundCallstack = GetCallstack(Found->EncodedCallstack);
+				const FString AllocCallstack = GetCallstack(Alloc.EncodedCallstack);
+
+				// Replace pointer.
+				AllocationMap.Add(Ptr, Alloc);
+				// Store the old pointer.
+				DuplicatedAllocMap.Add(Ptr, *Found);
+			}
+		}
+		else if( MemOp == EMemoryOperation::Free )
+		{
+			const FAllocationInfo* Found = AllocationMap.Find(Ptr);	
+			if( Found )
+			{
+				const bool bIsValid = Alloc.SequenceTag > Found->SequenceTag;
+				AllocationMap.Remove(Ptr);
+			}
+			else
+			{
+				FreeWithoutAllocMap.Add(Ptr,Alloc);
+				NumFWAMemoryOperations++;
+			}
+		}
+	}
+
+	UE_LOG(LogStats, Warning, TEXT("NumDuplicatedMemoryOperations %i"), NumDuplicatedMemoryOperations );
+	UE_LOG(LogStats, Warning, TEXT("NumFWAMemoryOperations %i"), NumFWAMemoryOperations );
+
+	// Dump problematic allocations
+	DuplicatedAllocMap.ValueSort( FAllocationInfoGreater() );
+	//FreeWithoutAllocMap
+
+	uint64 TotalDuplicatedMemory = 0;
+	for( const auto& It : DuplicatedAllocMap )
+	{
+		const FAllocationInfo& Alloc = It.Value;
+		TotalDuplicatedMemory += Alloc.Size;
+	}
+
+	const float MaxPctDisplayed = 0.80f;
+	uint64 DisplayedSoFar = 0;
+	for( const auto& It : DuplicatedAllocMap )
+	{
+		const FAllocationInfo& Alloc = It.Value;
+		const FString AllocCallstack = GetCallstack(Alloc.EncodedCallstack);
+		UE_LOG(LogStats, Warning, TEXT("%lli (%.2f MB) %s"), Alloc.Size, Alloc.Size/1024.0f/1024.0f, *AllocCallstack );
+
+		DisplayedSoFar += Alloc.Size;
+
+		const float CurrentPct = (float)DisplayedSoFar/(float)TotalDuplicatedMemory;
+		if( CurrentPct > MaxPctDisplayed )
+		{
+			break;
+		}
+	}
+
+	GenerateMemoryUsageReport( AllocationMap );
+}
+
+void FRawProfilerSession::GenerateMemoryUsageReport( const TMap<uint64, FAllocationInfo>& AllocationMap )
+{
+	if( AllocationMap.Num() == 0  )
+	{
+		UE_LOG(LogStats, Warning, TEXT("There are no allocations, make sure memory profiler is enabled") );
+	}
+	else
+	{
+		TMap<FName,FSizeAndCount> ScopedAllocations;
+
+		FName BrokenCallstackName = TEXT("245+12040+12265+14018+");
+		TArray<FAllocationInfo> BrokenAllocations;
+
+		uint64 TotalAllocatedMemory = 0;
+		for( const auto& It : AllocationMap )
+		{
+			const FAllocationInfo& Alloc = It.Value;
+			FSizeAndCount& SizeAndCount = ScopedAllocations.FindOrAdd(Alloc.EncodedCallstack);
+			SizeAndCount.Size += Alloc.Size;
+			SizeAndCount.Count += 1;
+
+			TotalAllocatedMemory += Alloc.Size;
+
+			if( Alloc.EncodedCallstack == BrokenCallstackName )
+			{
+				BrokenAllocations.Add(Alloc);
+			}
+		}
+
+		// Dump memory to the log.
+	
+		ScopedAllocations.ValueSort( FSizeAndCountGreater() );
+		BrokenAllocations.Sort( FAllocationInfoGreater() );
+
+		const float MaxPctDisplayed = 0.90f;
+		int32 CurrentIndex = 0;
+		uint64 DisplayedSoFar = 0;
+		UE_LOG( LogStats, Warning, TEXT("Index, Size (Size MB), Count, Stat desc") );
+		for( const auto& It : ScopedAllocations )
+		{
+			const FSizeAndCount& SizeAndCount = It.Value;
+			const FName& EncodedCallstack = It.Key;
+
+			const FString AllocCallstack = GetCallstack(EncodedCallstack);
+
+			/*TArray<FString> DecodedCallstack;
+			DecodeCallstack( EncodedCallstack, DecodedCallstack );
+
+			for( const auto& StrNameIndex : DecodedCallstack )
+			{
+				NAME_INDEX NameIndex = 0;
+				TTypeFromString<NAME_INDEX>::FromString( NameIndex, *StrNameIndex );
+
+				const FName LongName = FName(NameIndex, NameIndex, 0);
+
+				const FString ShortName = FStatNameAndInfo::GetShortNameFrom(LongName).ToString();
+				const FString Group = FStatNameAndInfo::GetGroupNameFrom(LongName).ToString();
+				FString Desc = FStatNameAndInfo::GetDescriptionFrom(LongName);
+				Desc.Trim();
+
+				if( Desc != ShortName )
+				{
+					if( Desc.Len() )
+					{
+						Desc += TEXT( " - " );
+					}
+					Desc += ShortName;
+				}
+
+				UE_LOG( LogStats, Warning, TEXT("%s [%s] "), 
+					*Desc, 
+					*Group );
+			}
+			}*/
+
+			UE_LOG( LogStats, Warning, TEXT("%2i, %llu (%.2f MB), %llu, %s"), 
+				CurrentIndex, 
+				SizeAndCount.Size, 
+				SizeAndCount.Size/1024.0f/1024.0f,
+				SizeAndCount.Count, 
+				*AllocCallstack );
+
+			CurrentIndex++;
+			DisplayedSoFar += SizeAndCount.Size;
+
+			const float CurrentPct = (float)DisplayedSoFar/(float)TotalAllocatedMemory;
+			if( CurrentPct > MaxPctDisplayed )
+			{
+				break;
+			}
+		}
+
+		UE_LOG(LogStats, Warning, TEXT("Allocated memory: %llu bytes (%.2f MB)"), TotalAllocatedMemory, TotalAllocatedMemory/1024.0f/1024.0f );
+	}
+
+	FPlatformMisc::RequestExit(true);
+}
\ No newline at end of file
diff --git a/Engine/Source/Developer/Profiler/Private/ProfilerSession.h b/Engine/Source/Developer/Profiler/Private/ProfilerSession.h
--- a/Engine/Source/Developer/Profiler/Private/ProfilerSession.h
+++ b/Engine/Source/Developer/Profiler/Private/ProfilerSession.h
@@ -1222,6 +1222,7 @@ public:
 	Profiler session for the raw stats files
 -----------------------------------------------------------------------------*/
 
+struct FAllocationInfo;
 
 class FRawProfilerSession : public FProfilerSession
 {
@@ -1238,6 +1239,12 @@ class FRawProfilerSession : public FProfilerSession
 	/** Index of the last processed data for the mini-view. */
 	int32 CurrentMiniViewFrame;
 
+	/** Basic memory profiling, only for debugging purpose. */
+	void ProcessMemoryOperations( const TMap<int64, FStatPacketArray>& CombinedHistory );
+
+	/** Generate a basic memory usage report and prints it to the log. */
+	void GenerateMemoryUsageReport( const TMap<uint64, FAllocationInfo>& AllocationMap );
+
 public:
 	/**
 	* Default constructor, creates a profiler session from a capture file.
diff --git a/Engine/Source/Developer/ProfilerClient/Private/ProfilerClientManager.cpp b/Engine/Source/Developer/ProfilerClient/Private/ProfilerClientManager.cpp
--- a/Engine/Source/Developer/ProfilerClient/Private/ProfilerClientManager.cpp
+++ b/Engine/Source/Developer/ProfilerClient/Private/ProfilerClientManager.cpp
@@ -443,7 +443,6 @@ void FServiceConnection::Initialize( const FProfilerServiceAuthorize2& Message,
 void FProfilerClientManager::HandleServiceMetaDataMessage( const FProfilerServiceMetaData& Message, const IMessageContextRef& Context )
 {
 #if STATS
-	// @TODO yrx 2014-04-14 Not used.
 	if (ActiveSessionId.IsValid() && Connections.Find(Message.InstanceId) != nullptr)
 	{
 		FServiceConnection& Connection = *Connections.Find(Message.InstanceId);
diff --git a/Engine/Source/Developer/ProfilerService/Private/ProfilerServiceManager.cpp b/Engine/Source/Developer/ProfilerService/Private/ProfilerServiceManager.cpp
--- a/Engine/Source/Developer/ProfilerService/Private/ProfilerServiceManager.cpp
+++ b/Engine/Source/Developer/ProfilerService/Private/ProfilerServiceManager.cpp
@@ -3,6 +3,7 @@
 #include "ProfilerServicePrivatePCH.h"
 
 #include "StatsData.h"
+#include "StatsFile.h"
 #include "SecureHash.h"
 
 
@@ -311,7 +312,6 @@ protected:
 FProfilerServiceManager::FProfilerServiceManager()
 {
 	MetaData.SecondsPerCycle = FPlatformTime::GetSecondsPerCycle();
-	NewMetaData.SecondsPerCycle = FPlatformTime::GetSecondsPerCycle();
 	PingDelegate = FTickerDelegate::CreateRaw(this, &FProfilerServiceManager::HandlePing);
 	DataFrame.Frame = 0;
 }
@@ -354,20 +354,7 @@ void FProfilerServiceManager::SendData(FProfilerCycleGraph& Data)
 void FProfilerServiceManager::StartCapture()
 {
 #if STATS
-	// fire off the equivalent of the stat startfile command
-	if (!Archive.IsValid())
-	{
-		// @TODO yrx 2014-06-05 Needs to be done on the stats thread via the task graph task.
-		FString Filename = CreateProfileFilename( FStatConstants::StatsFileExtension, true );
-		LastStatsFilename = FApp::GetInstanceName() + TEXT("_") + Filename;
-		TSharedPtr<FStatsWriteFile, ESPMode::ThreadSafe> ArchivePtr = MakeShareable(new FStatsWriteFile());
-		Archive = ArchivePtr;
-		Archive->Start( LastStatsFilename, false );
-		if (!Archive->IsValid())
-		{
-			Archive = nullptr;
-		}
-	}
+	DirectStatsCommand(TEXT("stat startfile"));
 #endif
 }
 
@@ -375,47 +362,12 @@ void FProfilerServiceManager::StartCapture()
 void FProfilerServiceManager::StopCapture()
 {
 #if STATS
-	if (Archive.IsValid())
-	{
-		Archive->Stop();
-	}
-	Archive = nullptr;
+	DirectStatsCommand(TEXT("stat stopfile"),true);
+	// Not thread-safe, but in this case it is ok, because we are waiting for completion.
+	LastStatsFilename = FCommandStatsFile::Get().LastFileSaved;
 #endif
 }
 
-void FProfilerServiceManager::UpdateMetaData()
-{
-	// @TODO yrx 2014-04-13 Obsolete, remove later.
-	// update the thread descriptions only if there has been a change
-	int32 OldCount = FRunnableThread::GetThreadRegistry().GetThreadCount();
-	if (FRunnableThread::GetThreadRegistry().IsUpdated())
-	{
-		FRunnableThread::GetThreadRegistry().Lock();
-		for( auto Iter = FRunnableThread::GetThreadRegistry().CreateConstIterator(); Iter; ++Iter)
-		{
-			if (!MetaData.ThreadDescriptions.Contains(Iter.Key()))
-			{
-				MetaData.ThreadDescriptions.Add(Iter.Key(), Iter.Value()->GetThreadName());
-				NewMetaData.ThreadDescriptions.Add(Iter.Key(), Iter.Value()->GetThreadName());
-			}
-		}
-		FRunnableThread::GetThreadRegistry().ClearUpdated();
-		FRunnableThread::GetThreadRegistry().Unlock();
-	}
-
-	if (MessageEndpoint.IsValid() && PreviewClients.Num() > 0 && (NewMetaData.GroupDescriptions.Num() > 0 || NewMetaData.StatDescriptions.Num() > 0 || NewMetaData.ThreadDescriptions.Num() > 0))
-	{
-		FArrayWriter ArrayWriter(true);
-		ArrayWriter << NewMetaData;
-		MessageEndpoint->Send(new FProfilerServiceMetaData(InstanceId, ArrayWriter), PreviewClients);
-	}
-
-	NewMetaData.GroupDescriptions.Reset();
-	NewMetaData.StatDescriptions.Reset();
-	NewMetaData.ThreadDescriptions.Reset();
-}
-
-
 void FProfilerServiceManager::SendMetaData(const FMessageAddress& client)
 {
 	if (MessageEndpoint.IsValid())
@@ -475,28 +427,12 @@ void FProfilerServiceManager::Init()
 		MessageEndpoint->Subscribe<FProfilerServiceUnsubscribe>();
 	}
 
-#if STATS
-	// check to see if we need to capture, specified from the command line
-	Archive = nullptr;
-#endif
-	if (FParse::Param(FCommandLine::Get(), TEXT("StartCapture")))
-	{
-		StartCapture();
-	}
-
 	FileTransferRunnable = new FFileTransferRunnable( MessageEndpoint );
 }
 
 
 void FProfilerServiceManager::Shutdown()
 {
-#if STATS
-	if (Archive.IsValid())
-	{
-		Archive = nullptr;
-	}
-#endif
-
 	delete FileTransferRunnable;
 	FileTransferRunnable = nullptr;
 
@@ -516,10 +452,27 @@ IProfilerServiceManagerPtr FProfilerServiceManager::CreateSharedServiceManager()
 	return ProfilerServiceManager;
 }
 
+void FProfilerServiceManager::AddNewFrameHandleStatsThread()
+{
+#if	STATS
+	const FStatsThreadState& Stats = FStatsThreadState::GetLocalState();
+	NewFrameDelegateHandle = Stats.NewFrameDelegate.AddRaw( this, &FProfilerServiceManager::HandleNewFrame );
+	StatsMasterEnableAdd();
+#endif // STATS
+}
+
+void FProfilerServiceManager::RemoveNewFrameHandleStatsThread()
+{
+#if	STATS
+	const FStatsThreadState& Stats = FStatsThreadState::GetLocalState();
+	Stats.NewFrameDelegate.Remove( NewFrameDelegateHandle );
+	StatsMasterEnableSubtract();
+#endif // STATS
+}
+
 void FProfilerServiceManager::SetPreviewState( const FMessageAddress& ClientAddress, const bool bRequestedPreviewState )
 {
 #if STATS
-	const FStatsThreadState& Stats = FStatsThreadState::GetLocalState();
 	FClientData* Client = ClientData.Find( ClientAddress );
 	if( Client )
 	{
@@ -532,15 +485,19 @@ void FProfilerServiceManager::SetPreviewState( const FMessageAddress& ClientAddr
 				// enable stat capture
 				if (PreviewClients.Num() == 0)
 				{
-					HandleNewFrameDelegateHandle = Stats.NewFrameDelegate.AddRaw( this, &FProfilerServiceManager::HandleNewFrame );
-					StatsMasterEnableAdd();
+					FGraphEventRef CompletionEvent = FSimpleDelegateGraphTask::CreateAndDispatchWhenReady
+					(
+						FSimpleDelegateGraphTask::FDelegate::CreateRaw( this, &FProfilerServiceManager::AddNewFrameHandleStatsThread ),
+						TStatId(), nullptr,
+						FPlatformProcess::SupportsMultithreading() ? ENamedThreads::StatsThread : ENamedThreads::GameThread
+					);
 				}
 				PreviewClients.Add(ClientAddress);
 				Client->Preview = true;
 				if (MessageEndpoint.IsValid())
 				{
-					Client->CurrentFrame = Stats.CurrentGameFrame;
-					MessageEndpoint->Send( new FProfilerServicePreviewAck( InstanceId, Stats.CurrentGameFrame ), ClientAddress );
+					Client->CurrentFrame = FStats::GameThreadStatsFrame;
+					MessageEndpoint->Send( new FProfilerServicePreviewAck( InstanceId, FStats::GameThreadStatsFrame ), ClientAddress );
 				}
 				SendMetaData(ClientAddress);
 			}
@@ -553,8 +510,13 @@ void FProfilerServiceManager::SetPreviewState( const FMessageAddress& ClientAddr
 				if (PreviewClients.Num() == 0)
 				{
 					// disable stat capture
-					Stats.NewFrameDelegate.Remove( HandleNewFrameDelegateHandle );
-					StatsMasterEnableAdd();
+					FGraphEventRef CompletionEvent = FSimpleDelegateGraphTask::CreateAndDispatchWhenReady
+					(
+						FSimpleDelegateGraphTask::FDelegate::CreateRaw( this, &FProfilerServiceManager::RemoveNewFrameHandleStatsThread ),
+						TStatId(), nullptr,
+						FPlatformProcess::SupportsMultithreading() ? ENamedThreads::StatsThread : ENamedThreads::GameThread
+					);
+					
 				}
 			}
 		}
@@ -590,13 +552,6 @@ bool FProfilerServiceManager::HandlePing( float DeltaTime )
 		}
 	}
 
-	if (PreviewClients.Num() == 0)
-	{
-		// disable stat capture
-		FStatsThreadState::GetLocalState().NewFrameDelegate.Remove(HandleNewFrameDelegateHandle);
-		StatsMasterEnableAdd();
-	}
-
 	// send the ping message
 	if (MessageEndpoint.IsValid() && Clients.Num() > 0)
 	{
@@ -611,16 +566,16 @@ void FProfilerServiceManager::HandleServiceCaptureMessage( const FProfilerServic
 {
 #if STATS
 	const bool bRequestedCaptureState = Message.bRequestedCaptureState;
-	const bool bIsCapturing = Archive.IsValid();
+	const bool bIsCapturing = FCommandStatsFile::Get().IsStatFileActive();
 
 	if( bRequestedCaptureState != bIsCapturing )
 	{
-		if( bRequestedCaptureState && !Archive.IsValid() )
+		if( bRequestedCaptureState && !bIsCapturing )
 		{
 			UE_LOG(LogProfile, Log, TEXT("StartCapture") );
 			StartCapture();
 		}
-		else if( !bRequestedCaptureState && Archive.IsValid() )
+		else if( !bRequestedCaptureState && bIsCapturing )
 		{
 			StopCapture();
 		}
@@ -692,7 +647,7 @@ void FProfilerServiceManager::HandleServiceSubscribeMessage( const FProfilerServ
 		FClientData Data;
 		Data.Active = true;
 		Data.Preview = false;
-		Data.StatsWriteFile.WriteHeader( false );
+		Data.StatsWriteFile.WriteHeader();
 
 		// add to the client list
 		ClientData.Add( MsgAddress, Data );
@@ -700,6 +655,7 @@ void FProfilerServiceManager::HandleServiceSubscribeMessage( const FProfilerServ
 		const TArray<uint8>& OutData = ClientData.Find( MsgAddress )->StatsWriteFile.GetOutData();
 		MessageEndpoint->Send( new FProfilerServiceAuthorize2( SessionId, InstanceId, OutData ), MsgAddress );
 
+		// Not thread-safe, but reading the number of elements should be ok.
 		const FStatsThreadState& Stats = FStatsThreadState::GetLocalState();
 		ClientData.Find( MsgAddress )->MetadataSize = Stats.ShortNameToLongName.Num();
 
@@ -740,6 +696,7 @@ void FProfilerServiceManager::HandleServiceUnsubscribeMessage( const FProfilerSe
 
 void FProfilerServiceManager::HandleNewFrame(int64 Frame)
 {
+	// Called from the stats thread.
 #if STATS
 	// package it up and send to the clients
 	if (MessageEndpoint.IsValid())
@@ -771,4 +728,4 @@ void FProfilerServiceManager::HandleNewFrame(int64 Frame)
 		}
 	}
 #endif
-}
+}
\ No newline at end of file
diff --git a/Engine/Source/Developer/ProfilerService/Private/ProfilerServiceManager.h b/Engine/Source/Developer/ProfilerService/Private/ProfilerServiceManager.h
--- a/Engine/Source/Developer/ProfilerService/Private/ProfilerServiceManager.h
+++ b/Engine/Source/Developer/ProfilerService/Private/ProfilerServiceManager.h
@@ -55,21 +55,10 @@ public:
 
 	virtual void SendData(FProfilerCycleGraph& Data) override;
 
-	virtual bool IsCapturing() const override
-	{ 
-#if STATS
-		return (PreviewClients.Num() > 0 || Archive.IsValid());
-#else
-		return false;
-#endif
-	}
-
 	virtual void StartCapture();
 	
 	virtual void StopCapture();
 
-	virtual void UpdateMetaData() override;
-
 	virtual void StartFrame(uint32 FrameNumber, double FrameStart) override;
 
 	virtual FStatMetaData& GetStatMetaData() override
@@ -133,8 +122,12 @@ private:
 	// Handles FProfilerServiceUnsubscribe messages.
 	void HandleServiceUnsubscribeMessage( const FProfilerServiceUnsubscribe& Message, const IMessageContextRef& Context );
 
-	// Handle a new frame from the stats system
+	/** Handles a new frame from the stats system. Called from the stats thread. */
 	void HandleNewFrame(int64 Frame);
+	
+	void AddNewFrameHandleStatsThread();
+
+	void RemoveNewFrameHandleStatsThread();
 
 private:
 
@@ -151,11 +144,6 @@ private:
 	/** Holds the client data for registered clients */
 	TMap<FMessageAddress, FClientData> ClientData;
 
-	/** Data writer for recording the data to disk */
-#if STATS
-	TSharedPtr<FStatsWriteFile, ESPMode::ThreadSafe> Archive;
-#endif
-
 	/** Thread used to read, prepare and send file chunks through the message bus. */
 	class FFileTransferRunnable* FileTransferRunnable;
 
@@ -164,7 +152,6 @@ private:
 
 	/** Stat meta data */
 	FStatMetaData MetaData;
-	FStatMetaData NewMetaData;
 
 	/** Delegate for notifying clients of received data */
 	FProfilerDataDelegate ProfilerDataDelegate;
@@ -179,5 +166,5 @@ private:
 	FDelegateHandle PingDelegateHandle;
 
 	/** Handle to the registered HandleNewFrame delegate */
-	FDelegateHandle HandleNewFrameDelegateHandle;
+	FDelegateHandle NewFrameDelegateHandle;
 };
diff --git a/Engine/Source/Developer/ProfilerService/Public/Interfaces/IProfilerServiceManager.h b/Engine/Source/Developer/ProfilerService/Public/Interfaces/IProfilerServiceManager.h
--- a/Engine/Source/Developer/ProfilerService/Public/Interfaces/IProfilerServiceManager.h
+++ b/Engine/Source/Developer/ProfilerService/Public/Interfaces/IProfilerServiceManager.h
@@ -501,20 +501,12 @@ public:
 	 */
 	virtual void SendData(FProfilerCycleGraph& Data) = 0;
 
-	/**
-	 * Determines if we are capturing data for clients.
-	 */
-	virtual bool IsCapturing() const = 0;
-
 	/** Starts a file capture. */
 	virtual void StartCapture() = 0;
 
 	/** Stops a file capture. */
 	virtual void StopCapture() = 0;
 
-	/** Updates the meta data. */
-	virtual void UpdateMetaData() = 0;
-
 	/** Starts a new frame of data and sends the previous if it exists. */
 	virtual void StartFrame(uint32 FrameNumber, double FrameStart) = 0;
 
diff --git a/Engine/Source/Runtime/Core/Private/HAL/UnrealMemory.cpp b/Engine/Source/Runtime/Core/Private/HAL/UnrealMemory.cpp
--- a/Engine/Source/Runtime/Core/Private/HAL/UnrealMemory.cpp
+++ b/Engine/Source/Runtime/Core/Private/HAL/UnrealMemory.cpp
@@ -41,11 +41,15 @@ void GCreateMalloc()
 		GMalloc = new FMallocThreadSafeProxy( GMalloc );
 	}
 
-
-
+	
 #if	STATS
+	const bool bCreateMallocProfiler = FParse::Param( FPlatformMisc::GetCommandLine(), TEXT( "MemoryProfiler" ) );
+
 	// Create the stats malloc profiler proxy, needs to be enabled separately.
-	//GMalloc = FStatsMallocProfilerProxy::Get();
+	if( bCreateMallocProfiler )
+	{
+		GMalloc = FStatsMallocProfilerProxy::Get();
+	}
 #endif // STATS
 }
 
diff --git a/Engine/Source/Runtime/Core/Private/Stats/Stats2.cpp b/Engine/Source/Runtime/Core/Private/Stats/Stats2.cpp
--- a/Engine/Source/Runtime/Core/Private/Stats/Stats2.cpp
+++ b/Engine/Source/Runtime/Core/Private/Stats/Stats2.cpp
@@ -28,15 +28,13 @@ DECLARE_DWORD_COUNTER_STAT( TEXT("Frame Packets Received"),STAT_StatFramePackets
 DECLARE_CYCLE_STAT(TEXT("WaitForStats"),STAT_WaitForStats,STATGROUP_Engine);
 DECLARE_CYCLE_STAT(TEXT("StatsNew Tick"),STAT_StatsNewTick,STATGROUP_StatSystem);
 DECLARE_CYCLE_STAT(TEXT("Parse Meta"),STAT_StatsNewParseMeta,STATGROUP_StatSystem);
+DECLARE_CYCLE_STAT(TEXT("Scan For Advance"),STAT_ScanForAdvance,STATGROUP_StatSystem);
 DECLARE_CYCLE_STAT(TEXT("Add To History"),STAT_StatsNewAddToHistory,STATGROUP_StatSystem);
+DECLARE_CYCLE_STAT(TEXT("Flush Raw Stats"),STAT_FlushRawStats,STATGROUP_StatSystem);
 
-DECLARE_MEMORY_STAT( TEXT( "Stats Descriptions" ), STAT_StatDescMemory, STATGROUP_StatSystem );
-DECLARE_MEMORY_STAT( TEXT( "Memory Map" ), STAT_MemoryAllocationMap, STATGROUP_StatSystem );
+DECLARE_MEMORY_STAT( TEXT("Stats Descriptions"), STAT_StatDescMemory, STATGROUP_StatSystem );
 
 DEFINE_STAT(STAT_FrameTime);
-DEFINE_STAT(STAT_FPS);
-DEFINE_STAT(STAT_DrawStats);
-DEFINE_STAT(STAT_StatMessagesMemory);
 
 /*-----------------------------------------------------------------------------
 	FStats2
@@ -64,6 +62,10 @@ void FStats::AdvanceFrame( bool bDiscardCallstack, const FOnAdvanceRenderingThre
 	{
 		Frame = -GameThreadStatsFrame; // mark this as a bad frame
 	}
+
+	// Update the seconds per cycle.
+	SET_FLOAT_STAT( STAT_SecondsPerCycle, FPlatformTime::GetSecondsPerCycle() );
+
 	static FStatNameAndInfo Adv( NAME_AdvanceFrame, "", "", TEXT( "" ), EStatDataType::ST_int64, true, false );
 	FThreadStats::AddMessage( Adv.GetEncodedName(), EStatOperation::AdvanceFrameEventGameThread, Frame ); // we need to flush here if we aren't collecting stats to make sure the meta data is up to date
 	if( FPlatformProperties::IsServerOnly() )
@@ -135,6 +137,8 @@ check(InStatId != NAME_None);
 
 #include "TaskGraphInterfaces.h"
 #include "StatsData.h"
+#include "StatsFile.h"
+#include "StatsMallocProfilerProxy.h"
 
 TStatIdData TStatId::TStatId_NAME_None;
 
@@ -634,15 +638,19 @@ static TAutoConsoleVariable<int32> CVarDumpStatPackets(	TEXT("DumpStatPackets"),
 /** The rendering thread runnable object. */
 class FStatsThread : public FRunnable, FSingleThreadRunnable
 {
+	/** Array of stat packets, queued data to be processed on this thread. */
 	FStatPacketArray IncomingData; 
-	FRunnableThread* Thread;
+
+	/** Stats state. */
 	FStatsThreadState& State;
+
+	/** Whether we are ready to process the packets, sets by game or render packets. */
 	bool bReadyToProcess;
 public:
 
+	/** Default constructor. */
 	FStatsThread()
-		: Thread(nullptr)
-		, State(FStatsThreadState::GetLocalState())
+		: State(FStatsThreadState::GetLocalState())
 		, bReadyToProcess(false)
 	{
 		check(IsInGameThread());
@@ -656,6 +664,7 @@ public:
 		return this;
 	}
 
+	/** Attaches to the task graph stats thread, all processing will be handled by the task graph. */
 	virtual uint32 Run()
 	{
 		FTaskGraphInterface::Get().AttachToThread(ENamedThreads::StatsThread);
@@ -663,26 +672,46 @@ public:
 		return 0;
 	}
 
+	/** Tick function. */
 	virtual void Tick() override
 	{
-		// Ignore all memory stack inside the stats thread tick, not really sure about that.
-		FThreadStats* ThreadStats = FThreadStats::GetThreadStats();
-		FStatMessageLock MessageLock(ThreadStats->MemoryMessageScope);
-
 		static double LastTime = -1.0;
-		if (bReadyToProcess && FPlatformTime::Seconds() - LastTime > .005) // we won't process more than every 5ms
+		bool bShouldProcess = false;
+
+		if( FThreadStats::bIsRawStatsActive )
 		{
-			// Update the seconds per cycle.
-			SET_FLOAT_STAT( STAT_SecondsPerCycle, FPlatformTime::GetSecondsPerCycle() ); 
-			SET_MEMORY_STAT( STAT_MemoryAllocationMap, State.GetAllocations().GetAllocatedSize()+State.GetNonSequentialAllocations().GetAllocatedSize() );
-			IStatGroupEnableManager::Get().UpdateMemoryUsage();
+			// For raw stats we process every 24MB of packet data to minimize the stats messages memory usage.
+			//const bool bShouldProcessRawStats = IncomingData.Packets.Num() > 10;
+			const int32 MaxIncomingMessages = 24*1024*1024/sizeof(FStatMessage);
+			const int32 MaxIncomingPackets = 16;
 
+			int32 IncomingDataMessages = 0;
+			for( int32 Index = 0; Index < IncomingData.Packets.Num(); ++Index )
+			{
+				IncomingDataMessages += IncomingData.Packets[Index]->StatMessages.Num();
+			}
+
+			bShouldProcess = IncomingDataMessages > MaxIncomingMessages || IncomingData.Packets.Num() > MaxIncomingPackets;
+		}
+		else
+		{
+			// For regular stats we won't process more than every 5ms.
+			bShouldProcess = bReadyToProcess && FPlatformTime::Seconds() - LastTime > 0.005f; 
+		}
+
+		if( bShouldProcess )
+		{
 			SCOPE_CYCLE_COUNTER(STAT_StatsNewTick);
+
+			IStatGroupEnableManager::Get().UpdateMemoryUsage();
+			State.UpdateStatMessagesMemoryUsage();
+		
 			bReadyToProcess = false;
 			FStatPacketArray NowData; 
 			Exchange(NowData.Packets, IncomingData.Packets);
 			INC_DWORD_STAT_BY(STAT_StatFramePacketsRecv, NowData.Packets.Num());
 			{
+				SCOPE_CYCLE_COUNTER(STAT_StatsNewParseMeta);
 				TArray<FStatMessage> MetaMessages;
 				{
 					FScopeLock Lock(&FStartupMessages::Get().CriticalSection);
@@ -693,12 +722,22 @@ public:
 					State.ProcessMetaDataOnly(MetaMessages);
 				}
 			}
-			{
-				SCOPE_CYCLE_COUNTER(STAT_StatsNewParseMeta);
+			{		
+				SCOPE_CYCLE_COUNTER(STAT_ScanForAdvance);
 				State.ScanForAdvance(NowData);
 			}
+			
+			if( FThreadStats::bIsRawStatsActive )
+			{
+				// Process raw stats.
+				State.ProcessRawStats(NowData);
+				State.ResetRegularStats();
+			}
+			else
 			{
+				// Process regular stats.
 				SCOPE_CYCLE_COUNTER(STAT_StatsNewAddToHistory);
+				State.ResetRawStats();
 				State.AddToHistoryAndEmpty(NowData);
 			}
 			check(!NowData.Packets.Num());
@@ -706,13 +745,14 @@ public:
 		}
 	}
 
-
+	/** Accesses singleton. */
 	static FStatsThread& Get()
 	{
 		static FStatsThread Singleton;
 		return Singleton;
 	}
 
+	/** Received a stat packet from other thread and add to the processing queue. */
 	void StatMessage(FStatPacket* Packet)
 	{
 		if (CVarDumpStatPackets.GetValueOnAnyThread())
@@ -722,18 +762,17 @@ public:
 
 		bReadyToProcess = Packet->ThreadType != EThreadType::Other;
 		IncomingData.Packets.Add(Packet);
-		const uint32 PacketMemory = Packet->StatMessages.GetAllocatedSize();
-		INC_MEMORY_STAT_BY(STAT_StatMessagesMemory, PacketMemory);
+		State.NumStatMessages.Add(Packet->StatMessages.Num());
 
 		Tick();
 	}
 
+	/** Start a stats runnable thread. */
 	void Start()
 	{
-		Thread = FRunnableThread::Create(this, TEXT("StatsThread"), 512 * 1024, TPri_BelowNormal, FPlatformAffinity::GetStatsThreadMask());
+		FRunnableThread* Thread = FRunnableThread::Create(this, TEXT("StatsThread"), 512 * 1024, TPri_BelowNormal, FPlatformAffinity::GetStatsThreadMask());
 		check(Thread != NULL);
 	}
-
 };
 
 /*-----------------------------------------------------------------------------
@@ -774,12 +813,13 @@ FThreadStatsPool::FThreadStatsPool()
 {
 	for( int32 Index = 0; Index < NUM_ELEMENTS_IN_POOL; ++Index )
 	{
-		Pool.Push( new FThreadStats(true) );
+		Pool.Push( new FThreadStats(EConstructor::FOR_POOL) );
 	}
 }
 
 FThreadStats* FThreadStatsPool::GetFromPool()
 {
+	FPlatformMisc::MemoryBarrier();
 	FThreadStats* Result = new(Pool.Pop()) FThreadStats();
 	check(Result && "Increase NUM_ELEMENTS_IN_POOL");
 	return Result;
@@ -802,36 +842,28 @@ FThreadSafeCounter FThreadStats::MasterEnableUpdateNumber;
 FThreadSafeCounter FThreadStats::MasterDisableChangeTagLock;
 bool FThreadStats::bMasterEnable = false;
 bool FThreadStats::bMasterDisableForever = false;
+bool FThreadStats::bIsRawStatsActive = false;
 
 FThreadStats::FThreadStats():
-	CurrentGameFrame(FStatsThreadState::GetLocalState().CurrentGameFrame),
+	CurrentGameFrame(FStats::GameThreadStatsFrame),
 	ScopeCount(0), 
 	bWaitForExplicitFlush(0),
 	MemoryMessageScope(0),
+	bReentranceGuard(false),
 	bSawExplicitFlush(false)
 {
-	Packet.ThreadId = FPlatformTLS::GetCurrentThreadId();
-	if (Packet.ThreadId == GGameThreadId)
-	{
-		Packet.ThreadType = EThreadType::Game;
-	}
-	else if (Packet.ThreadId == GRenderThreadId)
-	{
-		Packet.ThreadType = EThreadType::Renderer;
-	}
-	else
-	{
-		Packet.ThreadType = EThreadType::Other;
-	}
+	Packet.SetThreadProperties();
+
 	check(TlsSlot);
 	FPlatformTLS::SetTlsValue(TlsSlot, this);
 }
 
-CORE_API FThreadStats::FThreadStats( bool ):
-	CurrentGameFrame(0),
+FThreadStats::FThreadStats( EConstructor ):
+	CurrentGameFrame(-1),
 	ScopeCount(0), 
 	bWaitForExplicitFlush(0),
 	MemoryMessageScope(0),
+	bReentranceGuard(false),
 	bSawExplicitFlush(false)
 {}
 
@@ -847,28 +879,42 @@ void FThreadStats::CheckEnable()
 	}
 }
 
-void FThreadStats::Flush(bool bHasBrokenCallstacks)
+void FThreadStats::Flush(bool bHasBrokenCallstacks /*= false*/, bool bForceFlush /*= false*/)
 {
-	FStatMessageLock MessageLock(MemoryMessageScope);
-
 	if (bMasterDisableForever)
 	{
 		Packet.StatMessages.Empty();
 		return;
 	}
 
-	if( Packet.ThreadType == EThreadType::Other )
+	if( bIsRawStatsActive )
 	{
-		FPlatformMisc::MemoryBarrier();
-		const bool bFrameHasChanged = FStats::GameThreadStatsFrame > CurrentGameFrame;
-		if( bFrameHasChanged )
-		{
-			CurrentGameFrame = FStats::GameThreadStatsFrame;
-		}
-		else
-		{
-			return;
-		}
+		FlushRawStats(bHasBrokenCallstacks, bForceFlush);
+	}
+	else
+	{
+		FlushRegularStats(bHasBrokenCallstacks, bForceFlush);
+	}
+}
+
+void FThreadStats::FlushRegularStats( bool bHasBrokenCallstacks, bool bForceFlush )
+{
+	enum
+	{
+		PRESIZE_MAX_NUM_ENTRIES = 10,
+		PRESIZE_MAX_SIZE = 256*1024,
+	};
+
+
+	// Sends all collected messages when:
+	// The current game frame has changed.
+	// This a force flush when we shutting down the thread stats.
+	// This is an explicit flush from the game thread or the render thread.
+	const bool bFrameHasChanged = DetectAndUpdateCurrentGameFrame();
+	const bool bSendStatPacket = bFrameHasChanged || bForceFlush || bSawExplicitFlush;
+	if( !bSendStatPacket )
+	{
+		return;
 	}
 
 	if (!ScopeCount && Packet.StatMessages.Num())
@@ -904,11 +950,106 @@ void FThreadStats::Flush(bool bHasBrokenCallstacks)
 		}
 
 		TGraphTask<FStatMessagesTask>::CreateTask().ConstructAndDispatchWhenReady(ToSend);
-		if (Packet.ThreadType != EThreadType::Other && bSawExplicitFlush)
+		UpdateExplicitFlush();
+	}
+}
+
+void FThreadStats::FlushRawStats( bool bHasBrokenCallstacks /*= false*/, bool bForceFlush /*= false*/ )
+{
+	if( bReentranceGuard )
+	{
+		return;
+	}
+	bReentranceGuard = true;
+	
+	enum
+	{
+		/** Maximum number of messages in the stat packet. */
+		MAX_RAW_MESSAGES_IN_PACKET = 1024*1024 / sizeof(FStatMessage),
+	};
+
+	// Sends all collected messages when:
+	// Number of messages is greater than MAX_RAW_MESSAGES_IN_PACKET.
+	// The current game frame has changed.
+	// This a force flush when we shutting down the thread stats.
+	// This is an explicit flush from the game thread or the render thread.
+	const bool bFrameHasChanged = DetectAndUpdateCurrentGameFrame();
+	const int32 NumMessages = Packet.StatMessages.Num();
+	if( NumMessages > MAX_RAW_MESSAGES_IN_PACKET || bFrameHasChanged || bForceFlush || bSawExplicitFlush )
+	{
+		SCOPE_CYCLE_COUNTER(STAT_FlushRawStats);
+
+		FStatPacket* ToSend = new FStatPacket(Packet);
+		Exchange(ToSend->StatMessages, Packet.StatMessages);
+		ToSend->bBrokenCallstacks = bHasBrokenCallstacks;
+
+		check(!Packet.StatMessages.Num());
+
+		TGraphTask<FStatMessagesTask>::CreateTask().ConstructAndDispatchWhenReady(ToSend);
+		UpdateExplicitFlush();
+
+		const float NumMessagesAsMB = NumMessages*sizeof(FStatMessage) / 1024.0f / 1024.0f;
+		if( NumMessages > 524288 )
 		{
-			bWaitForExplicitFlush = 1;
-			ScopeCount++; // prevent sends until the next explicit flush
+			UE_LOG( LogStats, Warning, TEXT( "FlushRawStats NumMessages: %i (%.2f MB), Thread: %u" ), NumMessages, NumMessagesAsMB, Packet.ThreadId );
 		}
+
+		UE_LOG( LogStats, Verbose, TEXT( "FlushRawStats NumMessages: %i (%.2f MB), Thread: %u" ), NumMessages, NumMessagesAsMB, Packet.ThreadId );
+	}
+	
+	bReentranceGuard = false;
+}
+
+void FThreadStats::CheckForCollectingStartupStats()
+{
+	FString CmdLine(FCommandLine::Get());
+	FString StatCmds(TEXT("-StatCmds="));
+	while (1)
+	{
+		FString Cmds;
+		if (!FParse::Value(*CmdLine, *StatCmds, Cmds, false))
+		{
+			break;
+		}
+		TArray<FString> CmdsArray;
+		Cmds.ParseIntoArray(&CmdsArray, TEXT( "," ), true);
+		for (int32 Index = 0; Index < CmdsArray.Num(); Index++)
+		{
+			FString StatCmd = FString("stat ") + CmdsArray[Index].Trim();
+			UE_LOG(LogStatGroupEnableManager, Log, TEXT("Sending Stat Command '%s'"), *StatCmd);
+			DirectStatsCommand(*StatCmd);
+		}
+		int32 Index = CmdLine.Find(*StatCmds);
+		ensure(Index >= 0);
+		if (Index == INDEX_NONE)
+		{
+			break;
+		}
+		CmdLine = CmdLine.Mid(Index + StatCmds.Len());
+	}
+	
+	if (FParse::Param(FCommandLine::Get(), TEXT("LoadTimeStats")))
+	{
+		DirectStatsCommand(TEXT("stat group enable LinkerLoad"));
+		DirectStatsCommand(TEXT("stat group enable AsyncLoad"));
+		DirectStatsCommand(TEXT("stat dumpsum -start -ms=250 -num=240"));
+	}
+	if (FParse::Param(FCommandLine::Get(), TEXT("LoadTimeFile")))
+	{
+		DirectStatsCommand(TEXT("stat group enable LinkerLoad"));
+		DirectStatsCommand(TEXT("stat group enable AsyncLoad"));
+		DirectStatsCommand(TEXT("stat startfile"));
+	}
+
+	// Now we can safely enable malloc profiler.
+	// @TODO yrx 2014-12-01 Investigate if we can enable it earlier.
+	const bool bEnableMallocProfiler = FParse::Param( FCommandLine::Get(), TEXT("MemoryProfiler") );
+	if( bEnableMallocProfiler )
+	{
+		// Enable all available groups and enable malloc profiler.
+		IStatGroupEnableManager::Get().StatGroupEnableManagerCommand( TEXT("all") );
+		FStatsMallocProfilerProxy::Get()->SetState( true );
+		DirectStatsCommand(TEXT("stat startfileraw"), true);
 	}
 }
 
@@ -945,54 +1086,17 @@ void FThreadStats::StartThread()
 	}
 	check(IsThreadingReady());
 	CheckEnable();
+
 	// Preallocate a bunch of FThreadStats to avoid dynamic memory allocation.
 	FThreadStatsPool::Get();
 
+	if( FThreadStats::WillEverCollectData() )
 	{
-		FString CmdLine(FCommandLine::Get());
-		FString StatCmds(TEXT("-StatCmds="));
-		while (1)
-		{
-			FString Cmds;
-			if (!FParse::Value(*CmdLine, *StatCmds, Cmds, false))
-			{
-				break;
-			}
-			TArray<FString> CmdsArray;
-			Cmds.ParseIntoArray(&CmdsArray, TEXT( "," ), true);
-			for (int32 Index = 0; Index < CmdsArray.Num(); Index++)
-			{
-				FString StatCmd = FString("stat ") + CmdsArray[Index].Trim();
-				UE_LOG(LogStatGroupEnableManager, Log, TEXT("Sending Stat Command '%s'"), *StatCmd);
-				DirectStatsCommand(*StatCmd);
-			}
-			int32 Index = CmdLine.Find(*StatCmds);
-			ensure(Index >= 0);
-			if (Index == INDEX_NONE)
-			{
-				break;
-			}
-			CmdLine = CmdLine.Mid(Index + StatCmds.Len());
-		}
-	}
-
-	if (FParse::Param(FCommandLine::Get(), TEXT("LoadTimeStats")))
-	{
-		DirectStatsCommand(TEXT("stat group enable LinkerLoad"));
-		DirectStatsCommand(TEXT("stat group enable AsyncLoad"));
-		DirectStatsCommand(TEXT("stat dumpsum -start -ms=250 -num=240"));
-	}
-	if (FParse::Param(FCommandLine::Get(), TEXT("LoadTimeFile")))
-	{
-		DirectStatsCommand(TEXT("stat group enable LinkerLoad"));
-		DirectStatsCommand(TEXT("stat group enable AsyncLoad"));
-		DirectStatsCommand(TEXT("stat startfile"));
+		FThreadStats::ExplicitFlush(); // flush the stats and set update the scope so we don't flush again until a frame update, this helps prevent fragmentation
 	}
+	FStartupMessages::Get().AddThreadMetadata( NAME_GameThread, FPlatformTLS::GetCurrentThreadId() );
 
-	if( FParse::Param(FCommandLine::Get(), TEXT("MemoryProfiler")) )
-	{
-		DirectStatsCommand(TEXT("stat memoryprofiler enable"), true);
-	}
+	CheckForCollectingStartupStats();
 }
 
 static FGraphEventRef LastFramesEvents[MAX_STAT_LAG];
@@ -1002,15 +1106,19 @@ void FThreadStats::StopThread()
 {
 	// Nothing to stop if it was never started
 	if (IsThreadingReady())
-	{
+	{	
+		// If we are writing stats data, stop it now.
+		DirectStatsCommand(TEXT("stat stopfile"), true);
+
 		FThreadStats::MasterDisableForever();
+
 		WaitForStats();
 		for (int32 Index = 0; Index < MAX_STAT_LAG; Index++)
 		{
 			LastFramesEvents[Index] = NULL;
 		}
 		FGraphEventRef QuitTask = TGraphTask<FReturnGraphTask>::CreateTask(NULL, ENamedThreads::GameThread).ConstructAndDispatchWhenReady(FPlatformProcess::SupportsMultithreading() ? ENamedThreads::StatsThread : ENamedThreads::GameThread);
-		FTaskGraphInterface::Get().WaitUntilTaskCompletes(QuitTask, ENamedThreads::GameThread_Local);	
+		FTaskGraphInterface::Get().WaitUntilTaskCompletes(QuitTask, ENamedThreads::GameThread_Local);
 	}
 }
 
diff --git a/Engine/Source/Runtime/Core/Private/Stats/StatsCommand.cpp b/Engine/Source/Runtime/Core/Private/Stats/StatsCommand.cpp
--- a/Engine/Source/Runtime/Core/Private/Stats/StatsCommand.cpp
+++ b/Engine/Source/Runtime/Core/Private/Stats/StatsCommand.cpp
@@ -753,7 +753,7 @@ struct FHUDGroupManager
 		check(NumFrames <= Params.MaxHistoryFrames.Get());
 		if( NumFrames > 0 )
 		{
-			FGameThreadHudData* ToGame = new FGameThreadHudData();
+			FGameThreadHudData* ToGame = new FGameThreadHudData(false);
 
 			// Copy the total stats stack to the history stats stack and clear all nodes' data and set data type to none.
 			// Called to maintain the hierarchy.
@@ -979,96 +979,6 @@ static void DumpCPU(int64 Frame)
 	StatsMasterEnableSubtract();
 }
 
-static void DumpMemory(int64 /*Frame*/)
-{
-	FStatsThreadState& Stats = FStatsThreadState::GetLocalState();
-	const TMap<uint64,FAllocationInfo>& Allocations = Stats.GetAllocations();
-	if( Allocations.Num() == 0  )
-	{
-		UE_LOG(LogStats, Warning, TEXT("There are no allocations, make sure memory profiler is enabled") );
-	}
-	else
-	{
-		const TMap<uint64,FAllocationInfoEx>& NonSequentialAllocations = Stats.GetNonSequentialAllocations();
-
-		struct FSizeAndCount
-		{
-			uint64 Size;
-			uint64 Count;
-			FSizeAndCount()
-				: Size(0)
-				, Count(0)
-			{}
-		};
-
-		TMap<FName,FSizeAndCount> ScopedAllocations;
-
-		uint64 TotalAllocatedMemory = 0;
-		for( const auto& It : Allocations )
-		{
-			const FAllocationInfo& Alloc = It.Value;
-			FSizeAndCount& SizeAndCount = ScopedAllocations.FindOrAdd(Alloc.Scope);
-			SizeAndCount.Size += Alloc.Size;
-			SizeAndCount.Count += 1;
-
-			TotalAllocatedMemory += Alloc.Size;
-		}
-
-		// Dump memory to the log.
-		struct FGreater
-		{
-			FORCEINLINE bool operator()(const FSizeAndCount& A, const FSizeAndCount& B) const { return B.Size < A.Size; }
-		};
-
-		ScopedAllocations.ValueSort( FGreater() );
-
-		const float MaxPctDisplayed = 0.99f;
-		int32 CurrentIndex = 0;
-		uint64 DisplayedSoFar = 0;
-		UE_LOG( LogStats, Warning, TEXT("Index, Size (Size MB), Count, Stat desc [ Group ]") );
-		for( const auto& It : ScopedAllocations )
-		{
-			const FName LongName = It.Key;
-			const FSizeAndCount& SizeAndCount = It.Value;
-
-			const FString ShortName = FStatNameAndInfo::GetShortNameFrom(LongName).ToString();
-			const FString Group = FStatNameAndInfo::GetGroupNameFrom(LongName).ToString();
-			FString Desc = FStatNameAndInfo::GetDescriptionFrom(LongName);
-			Desc.Trim();
-
-			if( Desc != ShortName )
-			{
-				if( Desc.Len() )
-				{
-					Desc += TEXT( " - " );
-				}
-				Desc += ShortName;
-			}
-
-			UE_LOG( LogStats, Warning, TEXT("%2i, %llu (%.2f MB), %llu, %s [%s] "), 
-				CurrentIndex, 
-				SizeAndCount.Size, 
-				SizeAndCount.Size/1024.0f/1024.0f,
-				SizeAndCount.Count,
-				*Desc, 
-				*Group );
-
-			CurrentIndex++;
-			DisplayedSoFar += SizeAndCount.Size;
-
-			const float CurrentPct = (float)DisplayedSoFar/(float)TotalAllocatedMemory;
-			if( CurrentPct > MaxPctDisplayed )
-			{
-				break;
-			}
-		}
-
-		UE_LOG(LogStats, Warning, TEXT("Allocated memory: %llu bytes (%.2f MB)"), TotalAllocatedMemory, TotalAllocatedMemory/1024.0f/1024.0f );
-	}
-
-	Stats.NewFrameDelegate.Remove(DumpMemoryDelegateHandle);
-}
-
 static struct FDumpMultiple* DumpMultiple = NULL;
 
 struct FDumpMultiple
@@ -1187,21 +1097,20 @@ static void PrintStatsHelpToOutputDevice( FOutputDevice& Ar )
 #endif // WITH_ENGINE
 
 	Ar.Log( TEXT("stat startfile - starts dumping a capture"));
-	Ar.Log( TEXT("stat stopfile - stops dumping a capture"));
+	Ar.Log( TEXT("stat stopfile - stops dumping a capture (regular, raw, memory)"));
 
 	Ar.Log( TEXT("stat startfileraw - starts dumping a raw capture"));
-	Ar.Log( TEXT("stat stopfileraw - stops dumping a raw capture"));
 
 	Ar.Log( TEXT("stat toggledebug - toggles tracking the most memory expensive stats"));
 
-	Ar.Log( TEXT("stat memoryprofiler enable - enables tracking all memory operations, run 'stat startfileraw' before"));
-	Ar.Log( TEXT("stat memoryprofiler disable - disables tracking all memory operations, run 'stat startfileraw' after"));
-	Ar.Log( TEXT("stat dumpmemory - dump basic memory statistics WIP, memoryprofiler must to be enabled"));
+	Ar.Log( TEXT("add -memoryprofiler in the command line to enable the memory profiling"));
+	Ar.Log( TEXT("stat stopfile - stops tracking all memory operations and writes the results to the file"));
 }
 
+// @TODO yrx 2014-12-01 Move to StatsFile.cpp/.h
 static void CommandTestFile()
 {
-	const FString& LastFileSaved = FCommandStatsFile::LastFileSaved;
+	const FString& LastFileSaved = FCommandStatsFile::Get().LastFileSaved;
 
 	FStatsThreadState Loaded( LastFileSaved );
 	if( Loaded.GetLatestValidFrame() < 0 )
@@ -1249,10 +1158,6 @@ static void StatCmd(FString InCmd)
 		StatsMasterEnableAdd();
 		DumpFrameDelegateHandle = Stats.NewFrameDelegate.AddStatic(&DumpFrame);
 	} 
-	else if( FParse::Command(&Cmd,TEXT("DumpMemory")) )
-	{
-		DumpMemoryDelegateHandle = Stats.NewFrameDelegate.AddStatic(&DumpMemory);
-	}
 	else if ( FParse::Command(&Cmd,TEXT("DUMPNONFRAME")) )
 	{
 		DumpNonFrame(Stats);
@@ -1326,15 +1231,39 @@ static void StatCmd(FString InCmd)
 	}
 	else if( FParse::Command( &Cmd, TEXT( "STARTFILE" ) ) )
 	{
-		FCommandStatsFile::Start( Cmd );
+		FString Filename;
+		FParse::Token( Cmd, Filename, false );
+		FCommandStatsFile::Get().Start( Filename );
 	}
 	else if( FParse::Command( &Cmd, TEXT( "StartFileRaw" ) ) )
 	{
-		FCommandStatsFile::StartRaw( Cmd );
+		FThreadStats::EnableRawStats();
+		FString Filename;
+		FParse::Token( Cmd, Filename, false );
+		FCommandStatsFile::Get().StartRaw( Filename );
 	}
-	else if( FParse::Command( &Cmd, TEXT( "STOPFILE" ) ) || FParse::Command( &Cmd, TEXT( "StopFileRaw" ) ) )
+	else if( FParse::Command( &Cmd, TEXT( "STOPFILE" ) ) 
+		|| FParse::Command( &Cmd, TEXT( "StopFileRaw" ) ) )
 	{
-		FCommandStatsFile::Stop();
+		// Stop writing to a file.
+		FCommandStatsFile::Get().Stop();
+		FThreadStats::DisableRawStats();
+
+		if( FStatsMallocProfilerProxy::Get()->GetState() )
+		{
+			// Disable memory profiler and restore default stats groups.
+			FStatsMallocProfilerProxy::Get()->SetState( false );		
+			IStatGroupEnableManager::Get().StatGroupEnableManagerCommand( TEXT("default") );
+		}
+
+		Stats.ResetStatsForRawStats();
+
+		// Disable displaying the raw stats memory overhead.
+		FSimpleDelegateGraphTask::CreateAndDispatchWhenReady
+		(
+			FSimpleDelegateGraphTask::FDelegate::CreateRaw(&FHUDGroupGameThreadRenderer::Get(), &FHUDGroupGameThreadRenderer::NewData, (FGameThreadHudData*)nullptr),
+			TStatId(), nullptr, ENamedThreads::GameThread
+		);
 	}
 	else if( FParse::Command( &Cmd, TEXT( "TESTFILE" ) ) )
 	{
@@ -1359,11 +1288,11 @@ static void StatCmd(FString InCmd)
 	}
 	else if ( FParse::Command( &Cmd, TEXT( "memoryprofiler" ) ) )
 	{
-		IStatGroupEnableManager::Get().StatGroupEnableManagerCommand( TEXT("enable LinkerLoad") );
-		IStatGroupEnableManager::Get().StatGroupEnableManagerCommand( TEXT("enable AsyncLoad") );
-
-		const bool bEnable = FParse::Command( &Cmd, TEXT( "enable" ) ) || FParse::Command( &Cmd, TEXT( "start" ) );
-		FStatsMallocProfilerProxy::Get()->SetState( bEnable );
+		if( FParse::Command( &Cmd, TEXT( "snapshot" ) ) )
+		{
+			// Put a snapshot marker in the raw stats.
+			// @TODO yrx 2014-12-03 
+		}
 	}
 	// @see FStatHierParams
 	else if ( FParse::Command( &Cmd, TEXT( "hier" ) ) )
@@ -1448,9 +1377,6 @@ bool DirectStatsCommand(const TCHAR* Cmd, bool bBlockForCompletion /*= false*/,
 		else if( FParse::Command(&TempCmd,TEXT("DUMPFRAME")) )
 		{
 		}
-		else if( FParse::Command(&TempCmd,TEXT("DumpMemory")) )
-		{
-		}
 		else if ( FParse::Command(&TempCmd,TEXT("DUMPNONFRAME")) )
 		{
 		}
diff --git a/Engine/Source/Runtime/Core/Private/Stats/StatsData.cpp b/Engine/Source/Runtime/Core/Private/Stats/StatsData.cpp
--- a/Engine/Source/Runtime/Core/Private/Stats/StatsData.cpp
+++ b/Engine/Source/Runtime/Core/Private/Stats/StatsData.cpp
@@ -13,6 +13,7 @@ DECLARE_CYCLE_STAT(TEXT("Condense"),STAT_StatsCondense,STATGROUP_StatSystem);
 DECLARE_DWORD_COUNTER_STAT(TEXT("Frame Messages"),STAT_StatFrameMessages,STATGROUP_StatSystem);
 DECLARE_DWORD_COUNTER_STAT(TEXT("Total Frame Packets"),STAT_StatFramePackets,STATGROUP_StatSystem);
 DECLARE_DWORD_COUNTER_STAT(TEXT("Frame Messages Condensed"),STAT_StatFramePacketsCondensed,STATGROUP_StatSystem);
+DECLARE_MEMORY_STAT( TEXT("Stat Messages"), STAT_StatMessagesMemory, STATGROUP_StatSystem );
 
 /*-----------------------------------------------------------------------------
 	FStatConstants
@@ -457,10 +458,11 @@ void FComplexRawStatStackNode::CopyExclusivesFromSelf()
 
 void FStatPacketArray::Empty()
 {
+	FStatsThreadState& State = FStatsThreadState::GetLocalState();
+
 	for (int32 Index = 0; Index < Packets.Num(); Index++)
 	{
-		const uint32 PacketMemory = Packets[Index]->StatMessages.GetAllocatedSize();
-		DEC_MEMORY_STAT_BY(STAT_StatMessagesMemory, PacketMemory);
+		State.NumStatMessages.Subtract(Packets[Index]->StatMessages.Num());
 		delete Packets[Index];
 	}
 	Packets.Empty();
@@ -472,13 +474,16 @@ FStatsThreadState::FStatsThreadState(int32 InHistoryFrames)
 	, MinFrameSeen(-1)
 	, LastFullFrameMetaAndNonFrame(-1)
 	, LastFullFrameProcessed(-1)
-	, bWasLoaded(false)
-	, bFindMemoryExtensiveStats(false)
 	, CurrentGameFrame(1)
 	, CurrentRenderFrame(1)
+	, MaxNumStatMessages(0)
+	, TotalNumStatMessages(0)
+	, bWasLoaded(false)
+	, bFindMemoryExtensiveStats(false)
 {
 }
 
+// FStatsFileState:: ????
 //FStatsThreadState::FStatsThreadState(FString const& Filename)
 // void FStatsThreadState::AddMessages(TArray<FStatMessage>& InMessages)
 // @see moved to StatsFile.cpp
@@ -585,6 +590,12 @@ void FStatsThreadState::ScanForAdvance(const FStatMessagesArray& Data)
 			CurrentRenderFrame = NewRenderFrame;
 		}
 	}
+
+	// We don't care about bad frame when the raw stats are active.
+	if( FThreadStats::bIsRawStatsActive )
+	{
+		BadFrames.Empty();
+	}
 }
 
 void FStatsThreadState::ScanForAdvance(FStatPacketArray& NewData)
@@ -597,8 +608,10 @@ void FStatsThreadState::ScanForAdvance(FStatPacketArray& NewData)
 	uint32 Count = 0;
 	for (int32 Index = 0; Index < NewData.Packets.Num(); Index++)
 	{
-		int64 FrameNum = NewData.Packets[Index]->ThreadType == EThreadType::Renderer ? CurrentRenderFrame : CurrentGameFrame;
-		NewData.Packets[Index]->Frame = FrameNum;
+
+		const int64 FrameNum = NewData.Packets[Index]->ThreadType == EThreadType::Renderer ? CurrentRenderFrame : CurrentGameFrame;
+		NewData.Packets[Index]->AssignFrame( FrameNum );
+		
 		const FStatMessagesArray& Data = NewData.Packets[Index]->StatMessages;
 		ScanForAdvance(Data);
 		Count += Data.Num();
@@ -696,7 +709,7 @@ void FStatsThreadState::AddToHistoryAndEmpty(FStatPacketArray& NewData)
 		Frame.Packets.Add(NewData.Packets[Index]);
 	}
 
-	NewData.Packets.Empty(); // don't delete the elements
+	NewData.RemovePtrsButNoData(); // don't delete the elements
 
 	// now deal with metadata and non-frame stats
 
@@ -754,18 +767,13 @@ void FStatsThreadState::AddToHistoryAndEmpty(FStatPacketArray& NewData)
 							NonFrameMessages = &Frame.Packets[NonFrameIndex]->StatMessages;
 						}
 						
-						{
-							FStatMessageLock MessageLock(ThreadStats->MemoryMessageScope);
-							new (*NonFrameMessages) FStatMessage(It.Value());
-						}
-						
+						new (*NonFrameMessages) FStatMessage(It.Value());	
 					}
 				}
 
 				if(NonFrameMessages)
 				{
-					const uint32 PacketMemory = NonFrameMessages->GetAllocatedSize();
-					INC_MEMORY_STAT_BY(STAT_StatMessagesMemory, PacketMemory);
+					NumStatMessages.Add(NonFrameMessages->Num());
 				}
 
 				GoodFrames.Add(FrameNum);
@@ -844,6 +852,101 @@ void FStatsThreadState::AddToHistoryAndEmpty(FStatPacketArray& NewData)
 }
 
 
+void FStatsThreadState::ProcessRawStats( FStatPacketArray& NewData )
+{
+	if( NewRawStatPacket.IsBound() )
+	{
+		// First process the enqueued raw stats.
+		for (int32 Index = 0; Index < StartupRawStats.Packets.Num(); Index++)
+		{
+			const FStatPacket* StatPacket = StartupRawStats.Packets[Index];
+			NewRawStatPacket.Broadcast(StatPacket);
+		}
+		StartupRawStats.Empty();
+
+		// Now, process the raw stats.
+		for (int32 Index = 0; Index < NewData.Packets.Num(); Index++)
+		{
+			const FStatPacket* StatPacket = NewData.Packets[Index];
+			NewRawStatPacket.Broadcast(StatPacket);
+		}
+
+		// Now delete all the data.
+		NewData.Empty();
+	}
+	else
+	{
+		// The delegate is not bound yet, so store the data, because we don't want to lose any data.
+		for (int32 Index = 0; Index < NewData.Packets.Num(); Index++)
+		{
+			FStatPacket* StatPacket = NewData.Packets[Index];
+			StartupRawStats.Packets.Add(StatPacket);
+		}
+
+		NewData.RemovePtrsButNoData(); // Don't delete the elements.
+	}
+}
+
+void FStatsThreadState::ResetRawStats()
+{
+	// We no longer need any startup raw data.
+	StartupRawStats.Empty();
+}
+
+void FStatsThreadState::ResetRegularStats()
+{
+	// We need to reset these values after switching from the raw stats to the regular.
+	// !!CAUTION!!
+	// This is a bit unsafe as we lose accumulator history.
+	// Cycle counters and general counters should be just fine.
+	LastFullFrameMetaAndNonFrame = -1;
+	LastFullFrameProcessed = -1;
+	History.Empty();
+	CondensedStackHistory.Empty();
+	GoodFrames.Empty();
+	BadFrames.Empty();
+}
+
+void FStatsThreadState::UpdateStatMessagesMemoryUsage()
+{
+	// .UpdateStatMessagesMemoryUsage
+	const int32 CurrentNumStatMessages = NumStatMessages.GetValue();
+	//MaxNumStatMessages = FMath::Max(MaxNumStatMessages,CurrentNumStatMessages);
+
+	if( CurrentNumStatMessages > MaxNumStatMessages )
+	{
+		MaxNumStatMessages = CurrentNumStatMessages;
+	}
+
+	TotalNumStatMessages += CurrentNumStatMessages;
+	SET_MEMORY_STAT( STAT_StatMessagesMemory, CurrentNumStatMessages*sizeof(FStatMessage) );
+
+	if( FThreadStats::bIsRawStatsActive )
+	{
+		FGameThreadHudData* ToGame = new FGameThreadHudData(true);
+
+		const double InvMB = 1.0f / 1024.0f / 1024.0f;
+
+		// Format lines to be displayed on the hud.
+		const FString Current = FString::Printf( TEXT("Current: %.1f"), InvMB*CurrentNumStatMessages*sizeof(FStatMessage) );
+		const FString Max = FString::Printf( TEXT("Max: %.1f"), InvMB*(int64)MaxNumStatMessages*sizeof(FStatMessage) );
+		const FString Total = FString::Printf( TEXT("Total: %.1f") , InvMB*TotalNumStatMessages*sizeof(FStatMessage));
+
+		UE_LOG(LogStats, Verbose, TEXT("%s, %s, %s"), *Current, *Max, *Total );
+
+		ToGame->GroupDescriptions.Add( TEXT("RawStats memory usage (MB)") );
+		ToGame->GroupDescriptions.Add( Current );
+		ToGame->GroupDescriptions.Add( Max );
+		ToGame->GroupDescriptions.Add( Total );
+
+		FSimpleDelegateGraphTask::CreateAndDispatchWhenReady
+			(
+			FSimpleDelegateGraphTask::FDelegate::CreateRaw(&FHUDGroupGameThreadRenderer::Get(), &FHUDGroupGameThreadRenderer::NewData, ToGame),
+			TStatId(), nullptr, ENamedThreads::GameThread
+			);
+	}
+}
+
 void FStatsThreadState::GetInclusiveAggregateStackStats(int64 TargetFrame, TArray<FStatMessage>& OutStats, IItemFiler* Filter, bool bAddNonStackStats /*= true*/) const
 {
 	struct FTimeInfo
@@ -1051,61 +1154,7 @@ void FStatsThreadState::GetRawStackStats(int64 TargetFrame, FRawStatStackNode& R
 				}
 				else if( Op == EStatOperation::Memory )
 				{
-					//const FStatMessage& ScopeStartMessage = *StartStack.Last();
-
-					// Experimental code used only to test the implementation.
-					// First memory operation is Alloc or Free
-					const uint64 EncodedPtr = Item.GetValue_Ptr();
-					const bool bIsAlloc = (EncodedPtr & (uint64)EMemoryOperation::Alloc) != 0;
-					const bool bIsFree = (EncodedPtr & (uint64)EMemoryOperation::Free) != 0;
-					const uint64 Ptr = EncodedPtr & ~(uint64)EMemoryOperation::Mask;
-					if( bIsAlloc )
-					{
-						// @see FStatsMallocProfilerProxy::TrackAlloc
-						// After alloc ptr message there is always alloc size message.
-						Index++;
-						const FStatMessage& AllocSizeMessage = Data[Index];
-						const int64 AllocSize = AllocSizeMessage.GetValue_int64();
-
-						// Check the non-sequential allocation map.
-						FAllocationInfoEx* NonSeqAllocationInfo = NonSequentialAllocationMap.Find(Ptr);
-						if( NonSeqAllocationInfo )
-						{
-							// We don't need to add this allocation again.
-							NonSequentialAllocationMap.Remove(Ptr);
-						}
-						else
-						{
-							FAllocationInfo* OldAllocationInfo = AllocationMap.Find(Ptr);
-							if( OldAllocationInfo != nullptr )
-							{
-								// Ignore for now...
-							}
-
-							// Add a new allocation.
-							AllocationMap.Add( Ptr, FAllocationInfo(AllocSize, Current->Meta.NameAndInfo.GetEncodedName()) );
-						}					
-					}
-					else if( bIsFree )
-					{
-						// Remove the allocation, check if exists.
-						const int32 NumRemoved = AllocationMap.Remove( Ptr );
-
-						// We removed non-existent allocation, 
-						// 1. the timestamp of the allocation is before we started memory profiling
-						// 2. the allocation has been done on different thread than we are currently processing
-						if( NumRemoved == 0 )
-						{
-							NonSequentialAllocationMap.Add(Ptr, FAllocationInfoEx(
-								0, 
-								Current->Meta.NameAndInfo.GetEncodedName(),
-								/*ScopeStartMessage.GetValue_int64()*/0));
-						}
-					}
-					else
-					{
-						UE_LOG(LogStats, Warning, TEXT("Pointer from a memory operation is invalid") );
-					}
+					// Should never happen.
 				}
 				else if (OutNonStackStats)
 				{
@@ -1212,6 +1261,7 @@ int64 FStatsThreadState::GetFastThreadFrameTimeInternal( int64 TargetFrame, int3
 			{
 				FStatMessage const& Item = Data[Index];
 				EStatOperation::Type Op = Item.NameAndInfo.GetField<EStatOperation>();
+				FName LongName = Item.NameAndInfo.GetRawName();
 				if (Op == EStatOperation::CycleScopeStart)
 				{
 					check(Item.NameAndInfo.GetFlag(EStatMetaFlags::IsCycle));
diff --git a/Engine/Source/Runtime/Core/Private/Stats/StatsFile.cpp b/Engine/Source/Runtime/Core/Private/Stats/StatsFile.cpp
--- a/Engine/Source/Runtime/Core/Private/Stats/StatsFile.cpp
+++ b/Engine/Source/Runtime/Core/Private/Stats/StatsFile.cpp
@@ -15,14 +15,14 @@ DECLARE_CYCLE_STAT( TEXT( "Stream File" ), STAT_StreamFile, STATGROUP_StatSystem
 DECLARE_CYCLE_STAT( TEXT( "Wait For Write" ), STAT_StreamFileWaitForWrite, STATGROUP_StatSystem );
 
 /*-----------------------------------------------------------------------------
-	FAsyncWriteWorker
+	FAsyncStatsWrite
 -----------------------------------------------------------------------------*/
 
 /**
 *	Helper class used to save the capture stats data via the background thread.
-*	CAUTION!! Can exist only one instance at the same time. Synchronized via EnsureCompletion.
+*	!!CAUTION!! Can exist only one instance at the same time. Synchronized via EnsureCompletion.
 */
-class FAsyncWriteWorker : public FNonAbandonableTask
+class FAsyncStatsWrite : public FNonAbandonableTask
 {
 public:
 	/**
@@ -31,21 +31,16 @@ public:
 	 *	But in this specific case it is.
 	 *	@see SendTask
 	 */
-	FStatsWriteFile* Outer;
+	IStatsWriteFile* Outer;
 
 	/** Data for the file. Moved via Exchange. */
 	TArray<uint8> Data;
 
-	/** Thread cycles for the last frame. Moved via Exchange. */
-	TMap<uint32, int64> ThreadCycles;
-
-
 	/** Constructor. */
-	FAsyncWriteWorker( FStatsWriteFile* InStatsWriteFile )
+	FAsyncStatsWrite( IStatsWriteFile* InStatsWriteFile )
 		: Outer( InStatsWriteFile )
 	{
 		Exchange( Data, InStatsWriteFile->OutData );
-		Exchange( ThreadCycles, InStatsWriteFile->ThreadCycles );
 	}
 
 	/** Write compressed data to the file. */
@@ -61,7 +56,7 @@ public:
 		FCompressedStatsData CompressedData( Data, Outer->CompressedData );
 		Ar << CompressedData;
 
-		Outer->FramesInfo.Add( FStatsFrameInfo( FrameFileOffset, ThreadCycles ) );
+		Outer->FinalizeSavingData(FrameFileOffset);
 	}
 
 	FORCEINLINE TStatId GetStatId() const
@@ -80,10 +75,12 @@ FStatsThreadState::FStatsThreadState( FString const& Filename )
 	, MinFrameSeen( -1 )
 	, LastFullFrameMetaAndNonFrame( -1 )
 	, LastFullFrameProcessed( -1 )
-	, bWasLoaded( true )
-	, bFindMemoryExtensiveStats( false )
 	, CurrentGameFrame( -1 )
 	, CurrentRenderFrame( -1 )
+	, MaxNumStatMessages(0)
+	, TotalNumStatMessages(0)
+	, bWasLoaded( true )
+	, bFindMemoryExtensiveStats( false )
 {
 	const int64 Size = IFileManager::Get().FileSize( *Filename );
 	if( Size < 4 )
@@ -112,51 +109,43 @@ FStatsThreadState::FStatsThreadState( FString const& Filename )
 	TArray<FStatMessage> Messages;
 	if( Stream.Header.bRawStatsFile )
 	{
-		const int64 CurrentFilePos = FileReader->Tell();
+		int32 NumReadStatPackets = 0;
 
 		// Read metadata.
 		TArray<FStatMessage> MetadataMessages;
 		Stream.ReadFNamesAndMetadataMessages( *FileReader, MetadataMessages );
 		ProcessMetaDataForLoad( MetadataMessages );
 
-		// Read frames offsets.
-		Stream.ReadFramesOffsets( *FileReader );
-
-		// Verify frames offsets.
-		for( int32 FrameIndex = 0; FrameIndex < Stream.FramesInfo.Num(); ++FrameIndex )
+		// Read the raw stats messages.
+		const bool bHasCompressedData = Stream.Header.HasCompressedData();
+		check(bHasCompressedData);
+
+		// Buffer used to store the compressed and decompressed data.
+		TArray<uint8> SrcData;
+		TArray<uint8> DestData;
+		FCompressedStatsData UncompressedData( SrcData, DestData );
+		while( true )
 		{
-			const int64 FrameFileOffset = Stream.FramesInfo[FrameIndex].FrameFileOffset;
-			FileReader->Seek( FrameFileOffset );
+			// Read the compressed data.		
+			*FileReader << UncompressedData;
 
-			int64 TargetFrame;
-			*FileReader << TargetFrame;
+			if( UncompressedData.HasReachedEndOfCompressedData() )
+			{
+				// EOF
+				break;
 		}
-		FileReader->Seek( Stream.FramesInfo[0].FrameFileOffset );
-
-		// Read the raw stats messages.
-		FStatPacketArray IncomingData;
-		for( int32 FrameIndex = 0; FrameIndex < Stream.FramesInfo.Num(); ++FrameIndex )
-		{
-			int64 TargetFrame;
-			*FileReader << TargetFrame;
 
-			int32 NumPackets;
-			*FileReader << NumPackets;
+			// Select the proper archive.
+			FMemoryReader MemoryReader( DestData, true );
+			FArchive& Archive = bHasCompressedData ? MemoryReader : *FileReader;
 
-			for( int32 PacketIndex = 0; PacketIndex < NumPackets; PacketIndex++ )
-			{
-				FStatPacket* ToRead = new FStatPacket();
-				Stream.ReadStatPacket( *FileReader, *ToRead, bIsFinalized );
-				IncomingData.Packets.Add( ToRead );
+			FStatPacket StatPacket;
+			Stream.ReadStatPacket( MemoryReader, StatPacket );
+			// Process the raw stat messages.
+			NumReadStatPackets++;
 			}
 	
-			FStatPacketArray NowData;
-			// This is broken, do not use.
-// 			Exchange( NowData.Packets, IncomingData.Packets );
-// 			ScanForAdvance( NowData );
-// 			AddToHistoryAndEmpty( NowData );
-// 			check( !NowData.Packets.Num() );
-		}
+		UE_LOG( LogStats, Log, TEXT( "File: %s has %i stat packets" ), *Filename, NumReadStatPackets );
 	}
 	else
 	{
@@ -267,10 +256,10 @@ void FStatsThreadState::ProcessMetaDataForLoad(TArray<FStatMessage>& Data)
 }
 
 /*-----------------------------------------------------------------------------
-	FStatsWriteFile
+	IStatsWriteFile
 -----------------------------------------------------------------------------*/
 
-FStatsWriteFile::FStatsWriteFile()
+IStatsWriteFile::IStatsWriteFile() 
 	: File( nullptr )
 	, AsyncTask( nullptr )
 {
@@ -278,13 +267,11 @@ FStatsWriteFile::FStatsWriteFile()
 	CompressedData.Reserve( EStatsFileConstants::MAX_COMPRESSED_SIZE );
 }
 
-
-void FStatsWriteFile::Start( FString const& InFilename, bool bIsRawStatsFile )
+void IStatsWriteFile::Start( const FString& InFilename )
 {
 	const FString PathName = *(FPaths::ProfilingDir() + TEXT( "UnrealStats/" ));
-
-	FString Filename = PathName + InFilename;
-	FString Path = FPaths::GetPath( Filename );
+	const FString Filename = PathName + InFilename;
+	const FString Path = FPaths::GetPath( Filename );
 	IFileManager::Get().MakeDirectory( *Path, true );
 
 	UE_LOG( LogStats, Log, TEXT( "Opening stats file: %s" ), *Filename );
@@ -297,18 +284,19 @@ void FStatsWriteFile::Start( FString const& InFilename, bool bIsRawStatsFile )
 	else
 	{
 		ArchiveFilename = Filename;
-		AddNewFrameDelegate();
-		WriteHeader( bIsRawStatsFile );
+		WriteHeader();
+		SetDataDelegate(true);
 		StatsMasterEnableAdd();
 	}
 }
 
-void FStatsWriteFile::Stop()
+
+void IStatsWriteFile::Stop()
 {
 	if( IsValid() )
 	{
 		StatsMasterEnableSubtract();
-		RemoveNewFrameDelegate();
+		SetDataDelegate(false);
 		SendTask();
 		SendTask();
 		Finalize();
@@ -318,11 +306,11 @@ void FStatsWriteFile::Stop()
 		File = nullptr;
 
 		UE_LOG( LogStats, Log, TEXT( "Wrote stats file: %s" ), *ArchiveFilename );
-		FCommandStatsFile::LastFileSaved = ArchiveFilename;
+		FCommandStatsFile::Get().LastFileSaved = ArchiveFilename;
 	}
 }
 
-void FStatsWriteFile::WriteHeader( bool bIsRawStatsFile )
+void IStatsWriteFile::WriteHeader()
 {
 	FMemoryWriter MemoryWriter( OutData, false, true );
 	FArchive& Ar = File ? *File : MemoryWriter;
@@ -332,9 +320,9 @@ void FStatsWriteFile::WriteHeader( bool bIsRawStatsFile )
 	Ar << Magic;
 
 	// Serialize dummy header, overwritten in Finalize.
-	Header.Version = EStatMagicWithHeader::VERSION_4;
+	Header.Version = EStatMagicWithHeader::VERSION_5;
 	Header.PlatformName = FPlatformProperties::PlatformName();
-	Header.bRawStatsFile = bIsRawStatsFile;
+	//Header.bRawStatsFile = bIsRawStatsFile;
 	Ar << Header;
 
 	// Serialize metadata.
@@ -342,35 +330,17 @@ void FStatsWriteFile::WriteHeader( bool bIsRawStatsFile )
 	Ar.Flush();
 }
 
-void FStatsWriteFile::WriteFrame( int64 TargetFrame, bool bNeedFullMetadata /*= false*/ )
+void IStatsWriteFile::WriteMetadata( FArchive& Ar )
 {
-	FMemoryWriter Ar( OutData, false, true );
-	const int64 PrevArPos = Ar.Tell();
-
-	if( bNeedFullMetadata )
-	{
-		WriteMetadata( Ar );
-	}
-
-	FStatsThreadState const& Stats = FStatsThreadState::GetLocalState();
-	TArray<FStatMessage> const& Data = Stats.GetCondensedHistory( TargetFrame );
-	for( auto It = Data.CreateConstIterator(); It; ++It )
+	const FStatsThreadState& Stats = FStatsThreadState::GetLocalState();
+	for( const auto& It : Stats.ShortNameToLongName )
 	{
-		WriteMessage( Ar, *It );
+		const FStatMessage& StatMessage = It.Value;
+		WriteMessage( Ar, StatMessage );
 	}
-
-	// Get cycles for all threads, so we can use that data to generate the mini-view.
-	for( auto It = Stats.Threads.CreateConstIterator(); It; ++It )
-	{
-		const int64 Cycles = Stats.GetFastThreadFrameTime( TargetFrame, It.Key() );
-		ThreadCycles.Add( It.Key(), Cycles );
 	}
 
-	// Serialize thread cycles. Disabled for now.
-	//Ar << ThreadCycles;
-}
-
-void FStatsWriteFile::Finalize()
+void IStatsWriteFile::Finalize()
 {
 	FArchive& Ar = *File;
 
@@ -381,6 +351,7 @@ void FStatsWriteFile::Finalize()
 
 	// Write out frame table and update header with offset and count.
 	Header.FrameTableOffset = Ar.Tell();
+	// This is ok to access the frames info, the async write thread is dead.
 	Ar << FramesInfo;
 
 	const FStatsThreadState& Stats = FStatsThreadState::GetLocalState();
@@ -430,91 +401,61 @@ void FStatsWriteFile::Finalize()
 	Ar << Header;
 }
 
-void FStatsWriteFile::NewFrame( int64 TargetFrame )
-{
-	SCOPE_CYCLE_COUNTER( STAT_StreamFile );
-
-	// @TODO yrx 2014-11-25 Add stat startfile -num=number of frames to capture
-	/*
-	enum
-	{
-		MAX_NUM_RAWFRAMES = 120,
-	};
-	if( Header.bRawStatsFile )
-	{
-		if( FCommandStatsFile::FirstFrame == -1 )
-		{
-			FCommandStatsFile::FirstFrame = TargetFrame;
-		}
-		else if( TargetFrame > FCommandStatsFile::FirstFrame + MAX_NUM_RAWFRAMES )
-		{
-			FCommandStatsFile::Stop();
-			FCommandStatsFile::FirstFrame = -1;
-			return;
-		}
-	}
-	}*/
-
-	WriteFrame( TargetFrame );
-	SendTask();
-}
-
-void FStatsWriteFile::SendTask()
+void IStatsWriteFile::SendTask()
 {
 	if( AsyncTask )
 	{
 		SCOPE_CYCLE_COUNTER( STAT_StreamFileWaitForWrite );
 		AsyncTask->EnsureCompletion();
 		delete AsyncTask;
-		AsyncTask = NULL;
+		AsyncTask = nullptr;
 	}
 	if( OutData.Num() )
 	{
-		AsyncTask = new FAsyncTask<FAsyncWriteWorker>( this );
+		AsyncTask = new FAsyncTask<FAsyncStatsWrite>( this );
 		check( !OutData.Num() );
-		check( !ThreadCycles.Num() );
+		//check( !ThreadCycles.Num() );
 		AsyncTask->StartBackgroundTask();
 	}
 }
 
-bool FStatsWriteFile::IsValid() const
-{
-	return !!File;
-}
+/*-----------------------------------------------------------------------------
+	FStatsWriteFile
+-----------------------------------------------------------------------------*/
+
 
-void FStatsWriteFile::AddNewFrameDelegate()
+void FStatsWriteFile::SetDataDelegate( bool bSet )
 {
 	FStatsThreadState const& Stats = FStatsThreadState::GetLocalState();
-	NewFrameDelegateHandle = Stats.NewFrameDelegate.AddThreadSafeSP( this->AsShared(), &FStatsWriteFile::NewFrame );
+	if( bSet )
+	{
+		DataDelegateHandle = Stats.NewFrameDelegate.AddRaw( this, &FStatsWriteFile::WriteFrame );
 }
-
-void FStatsWriteFile::RemoveNewFrameDelegate()
+	else
 {
-	FStatsThreadState const& Stats = FStatsThreadState::GetLocalState();
-	Stats.NewFrameDelegate.Remove( NewFrameDelegateHandle );
+		Stats.NewFrameDelegate.Remove( DataDelegateHandle );
+	}
 }
 
-/*-----------------------------------------------------------------------------
-	FRawStatsWriteFile
------------------------------------------------------------------------------*/
 
-void FRawStatsWriteFile::WriteFrame( int64 TargetFrame, bool bNeedFullMetadata /*= false */)
+void FStatsWriteFile::WriteFrame( int64 TargetFrame, bool bNeedFullMetadata )
 {
-	FMemoryWriter Ar( OutData, false, true );
-	const FStatsThreadState& Stats = FStatsThreadState::GetLocalState();
+	// @TODO yrx 2014-11-25 Add stat startfile -num=number of frames to capture
 
-	check( Stats.IsFrameValid( TargetFrame ) );
+	SCOPE_CYCLE_COUNTER( STAT_StreamFile );
 
-	Ar << TargetFrame;
+	FMemoryWriter Ar( OutData, false, true );
 
-	// Write stat packet.
-	const FStatPacketArray& Frame = Stats.GetStatPacketArray( TargetFrame );
-	int32 NumPackets = Frame.Packets.Num();
-	Ar << NumPackets;
-	for( int32 PacketIndex = 0; PacketIndex < Frame.Packets.Num(); PacketIndex++ )
+	if( bNeedFullMetadata )
+	{
+		WriteMetadata( Ar );
+	}
+
+	FStatsThreadState const& Stats = FStatsThreadState::GetLocalState();
+	TArray<FStatMessage> const& Data = Stats.GetCondensedHistory( TargetFrame );
+	for( auto It = Data.CreateConstIterator(); It; ++It )
 	{
-		FStatPacket& StatPacket = *Frame.Packets[PacketIndex];
-		WriteStatPacket( Ar, StatPacket );
+		WriteMessage( Ar, *It );
 	}
 
 	// Get cycles for all threads, so we can use that data to generate the mini-view.
@@ -524,53 +465,104 @@ void FRawStatsWriteFile::WriteFrame( int64 TargetFrame, bool bNeedFullMetadata /
 		ThreadCycles.Add( It.Key(), Cycles );
 	}
 	
-	// Serialize thread cycles.
-	Ar << ThreadCycles;
+	// Serialize thread cycles. Disabled for now.
+	//Ar << ThreadCycles;
+}
+
+void FStatsWriteFile::FinalizeSavingData( int64 FrameFileOffset )
+{
+	// Called from the async write thread.
+	FramesInfo.Add( FStatsFrameInfo( FrameFileOffset, ThreadCycles ) );
 }
 
 /*-----------------------------------------------------------------------------
-	Commands functionality
+	FRawStatsWriteFile
 -----------------------------------------------------------------------------*/
 
-FString FCommandStatsFile::LastFileSaved;
-int64 FCommandStatsFile::FirstFrame = -1;
-TSharedPtr<FStatsWriteFile, ESPMode::ThreadSafe> FCommandStatsFile::CurrentStatsFile = nullptr;
+void FRawStatsWriteFile::SetDataDelegate( bool bSet )
+{
+	FStatsThreadState const& Stats = FStatsThreadState::GetLocalState();
+	if( bSet )
+	{
+		DataDelegateHandle = Stats.NewRawStatPacket.AddRaw( this, &FRawStatsWriteFile::WriteRawStatPacket );
+		if( !bWrittenOffsetToData )
+		{
+			const int64 FrameFileOffset = File->Tell();
+			FramesInfo.Add( FStatsFrameInfo( FrameFileOffset ) );
+			bWrittenOffsetToData = true;
+		}
+	}
+	else
+	{
+		Stats.NewRawStatPacket.Remove( DataDelegateHandle );
+	}
+}
 
-void FCommandStatsFile::Start( const TCHAR* Cmd )
+void FRawStatsWriteFile::WriteRawStatPacket( const FStatPacket* StatPacket )
 {
-	CurrentStatsFile = NULL;
-	FString File;
-	FParse::Token( Cmd, File, false );
-	TSharedPtr<FStatsWriteFile, ESPMode::ThreadSafe> StatFile = MakeShareable( new FStatsWriteFile() );
-	CurrentStatsFile = StatFile;
-	CurrentStatsFile->Start( File, false );
-	if( !CurrentStatsFile->IsValid() )
+	FMemoryWriter Ar( OutData, false, true );
+
+	// Write stat packet.
+	WriteStatPacket( Ar, (FStatPacket&)*StatPacket );
+	SendTask();
+}
+
+void FRawStatsWriteFile::WriteStatPacket( FArchive& Ar, FStatPacket& StatPacket )
+{
+	Ar << StatPacket.Frame;
+	Ar << StatPacket.ThreadId;
+	int32 MyThreadType = (int32)StatPacket.ThreadType;
+	Ar << MyThreadType;
+
+	Ar << StatPacket.bBrokenCallstacks;
+	// We must handle stat messages in a different way.
+	int32 NumMessages = StatPacket.StatMessages.Num();
+	Ar << NumMessages;
+	for( int32 MessageIndex = 0; MessageIndex < NumMessages; ++MessageIndex )
 	{
-		CurrentStatsFile = nullptr;
+		WriteMessage( Ar, StatPacket.StatMessages[MessageIndex] );
 	}
 }
 
-void FCommandStatsFile::StartRaw( const TCHAR* Cmd )
+/*-----------------------------------------------------------------------------
+	Commands functionality
+-----------------------------------------------------------------------------*/;
+
+FCommandStatsFile& FCommandStatsFile::Get()
 {
-	CurrentStatsFile = NULL;
-	FString File;
-	FParse::Token( Cmd, File, false );
-	TSharedPtr<FStatsWriteFile, ESPMode::ThreadSafe> StatFile = MakeShareable( new FRawStatsWriteFile() );
-	CurrentStatsFile = StatFile;
-	CurrentStatsFile->Start( File, true );
-	if( !CurrentStatsFile->IsValid() )
+	static FCommandStatsFile CommandStatsFile;
+	return CommandStatsFile;
+}
+
+
+void FCommandStatsFile::Start( const FString& Filename )
 	{
-		CurrentStatsFile = nullptr;
+	Stop();
+	CurrentStatsFile = new FStatsWriteFile();
+	CurrentStatsFile->Start( Filename );
+
+	StatFileActiveCounter.Increment();
 	}
+
+void FCommandStatsFile::StartRaw( const FString& Filename )
+{
+	Stop();
+	CurrentStatsFile = new FRawStatsWriteFile();
+	CurrentStatsFile->Start( Filename );
+
+	StatFileActiveCounter.Increment();
 }
 
 void FCommandStatsFile::Stop()
 {
-	if( CurrentStatsFile.IsValid() )
+	if( CurrentStatsFile )
 	{
 		CurrentStatsFile->Stop();
+		delete CurrentStatsFile;
+		CurrentStatsFile = nullptr;
+
+		StatFileActiveCounter.Decrement();
 	}
-	CurrentStatsFile = nullptr;
 }
 
 #endif // STATS
diff --git a/Engine/Source/Runtime/Core/Private/Stats/StatsMallocProfilerProxy.cpp b/Engine/Source/Runtime/Core/Private/Stats/StatsMallocProfilerProxy.cpp
--- a/Engine/Source/Runtime/Core/Private/Stats/StatsMallocProfilerProxy.cpp
+++ b/Engine/Source/Runtime/Core/Private/Stats/StatsMallocProfilerProxy.cpp
@@ -18,9 +18,10 @@
 /** Fake stat group and memory stats. */
 DECLARE_STATS_GROUP(TEXT("Memory Profiler"), STATGROUP_MemoryProfiler, STATCAT_Advanced);
 
-DECLARE_MEMORY_PTR_STAT(TEXT("Memory Free Ptr"),			STAT_Memory_FreePtr,		STATGROUP_MemoryProfiler);
-DECLARE_MEMORY_PTR_STAT(TEXT("Memory Alloc Ptr"),			STAT_Memory_AllocPtr,		STATGROUP_MemoryProfiler);
-DECLARE_MEMORY_SIZE_STAT(TEXT("Memory Alloc Size"),			STAT_Memory_AllocSize,		STATGROUP_MemoryProfiler);
+DECLARE_MEMORY_PTR_STAT(TEXT("Memory Free Ptr"),				STAT_Memory_FreePtr,				STATGROUP_MemoryProfiler);
+DECLARE_MEMORY_PTR_STAT(TEXT("Memory Alloc Ptr"),				STAT_Memory_AllocPtr,				STATGROUP_MemoryProfiler);
+DECLARE_MEMORY_SIZE_STAT(TEXT("Memory Alloc Size"),				STAT_Memory_AllocSize,				STATGROUP_MemoryProfiler);
+DECLARE_MEMORY_SIZE_STAT(TEXT("Memory Operation Sequence Tag"),	STAT_Memory_OperationSequenceTag,	STATGROUP_MemoryProfiler);
 
 /** Stats for memory usage by the profiler. */
 DECLARE_DWORD_COUNTER_STAT(TEXT("Profiler AllocPtr Calls"),	STAT_Memory_AllocPtr_Calls,	STATGROUP_StatSystem);
@@ -28,9 +29,15 @@ DECLARE_DWORD_COUNTER_STAT(TEXT("Profiler FreePtr Calls"),	STAT_Memory_FreePtr_C
 DECLARE_MEMORY_STAT( TEXT( "Profiler AllocPtr" ),			STAT_Memory_AllocPtr_Mem,	STATGROUP_StatSystem );
 DECLARE_MEMORY_STAT( TEXT( "Profiler FreePtr" ),			STAT_Memory_FreePtr_Mem,	STATGROUP_StatSystem );
 
+//#define DEBUG_MALLOC_PROXY
+
+/** Debugging only. */
+CORE_API FThreadStats* GThreadStatsToDumpMemory = nullptr;
+
 FStatsMallocProfilerProxy::FStatsMallocProfilerProxy( FMalloc* InMalloc ) 
 	: UsedMalloc( InMalloc )
 	, bEnabled(false)
+	, bWasEnabled(false)
 {
 }
 
@@ -45,21 +52,27 @@ FStatsMallocProfilerProxy* FStatsMallocProfilerProxy::Get()
 }
 
 
-void FStatsMallocProfilerProxy::SetState( bool bEnable )
+void FStatsMallocProfilerProxy::SetState( bool bNewState )
 {
-	bEnabled = bEnable;
-	FPlatformMisc::MemoryBarrier();
-	if( bEnable )
+	if( !bWasEnabled && bNewState )
 	{
-		StatsMasterEnableAdd();
+		bEnabled = true;
+		bWasEnabled = true;
+		
+		UE_LOG( LogStats, Warning, TEXT( "Malloc profiler is enabled" ) );
 	}
-	else
+	else if( bWasEnabled && !bNewState )
 	{
-		StatsMasterEnableSubtract();
+		bEnabled = false;
+		UE_LOG( LogStats, Warning, TEXT( "Malloc profiler has been disabled, all data should be ready" ) );
 	}
+	else if( bWasEnabled )
+	{
+		UE_LOG( LogStats, Warning, TEXT( "Malloc profiler has already been stopped and cannot be restarted." ) );
+	}
+	FPlatformMisc::MemoryBarrier();
 }
 
-
 void FStatsMallocProfilerProxy::InitializeStatsMetadata()
 {
 	UsedMalloc->InitializeStatsMetadata();
@@ -70,6 +83,7 @@ void FStatsMallocProfilerProxy::InitializeStatsMetadata()
 	const FName NameAllocPtr = GET_STATFNAME(STAT_Memory_AllocPtr);
 	const FName NameFreePtr = GET_STATFNAME(STAT_Memory_FreePtr);
 	const FName NameAllocSize = GET_STATFNAME(STAT_Memory_AllocSize);
+	const FName NameOperationSequenceTag = GET_STATFNAME(STAT_Memory_OperationSequenceTag);
 
 	GET_STATFNAME(STAT_Memory_AllocPtr_Calls);
 	GET_STATFNAME(STAT_Memory_FreePtr_Calls);
@@ -82,8 +96,7 @@ void FStatsMallocProfilerProxy::InitializeStatsMetadata()
 	//StatPtr_STAT_Memory_AllocSize;
 }
 
-
-void FStatsMallocProfilerProxy::TrackAlloc( void* Ptr, int64 Size )
+void FStatsMallocProfilerProxy::TrackAlloc( void* Ptr, int64 Size, int32 SequenceTag )
 {
 	if( bEnabled )
 	{
@@ -91,17 +104,38 @@ void FStatsMallocProfilerProxy::TrackAlloc( void* Ptr, int64 Size )
 		{
 			FThreadStats* ThreadStats = FThreadStats::GetThreadStats();
 
+#ifdef DEBUG_MALLOC_PROXY
+			if( GThreadStatsToDumpMemory == ThreadStats && ThreadStats->MemoryMessageScope == 0 )
+			{
+				ThreadStats->MemoryMessageScope++;
+				UE_LOG( LogStats, Warning, TEXT( "TrackAlloc, %llu, %lli, %i, %i" ), (uint64)(UPTRINT)Ptr, Size, SequenceTag, ThreadStats->MemoryMessageScope );
+				ThreadStats->MemoryMessageScope--;
+			}
+#endif // DEBUG_MALLOC_PROXY
+		
 			if( ThreadStats->MemoryMessageScope == 0 )
 			{
-				ThreadStats->AddMemoryMessage( GET_STATFNAME(STAT_Memory_AllocPtr), (uint64)(UPTRINT)Ptr | (uint64)EMemoryOperation::Alloc );	// 16 bytes
-				ThreadStats->AddMemoryMessage( GET_STATFNAME(STAT_Memory_AllocSize), Size );					// 32 bytes total
+#if	_DEBUG
+				if( ThreadStats->Packet.StatMessages.Num() > 0 && ThreadStats->Packet.StatMessages.Num() % 32767 == 0 )
+				{
+					ThreadStats->MemoryMessageScope++;
+					const double InvMB = 1.0f / 1024.0f / 1024.0f;
+					UE_LOG( LogStats, Verbose, TEXT( "ThreadID: %i, Current: %.1f" ), FPlatformTLS::GetCurrentThreadId(), InvMB*(int64)ThreadStats->Packet.StatMessages.Num()*sizeof( FStatMessage ) );
+					ThreadStats->MemoryMessageScope--;
+				}
+#endif // _DEBUG
+
+				// 48 bytes per allocation.
+				ThreadStats->AddMemoryMessage( GET_STATFNAME(STAT_Memory_AllocPtr), (uint64)(UPTRINT)Ptr | (uint64)EMemoryOperation::Alloc );
+				ThreadStats->AddMemoryMessage( GET_STATFNAME(STAT_Memory_AllocSize), Size );			
+				ThreadStats->AddMemoryMessage( GET_STATFNAME(STAT_Memory_OperationSequenceTag), (int64)SequenceTag );			
 				AllocPtrCalls.Increment();
 			}
 		}	
 	}
 }
 
-void FStatsMallocProfilerProxy::TrackFree( void* Ptr )
+void FStatsMallocProfilerProxy::TrackFree( void* Ptr, int32 SequenceTag )
 {
 	if( bEnabled )
 	{
@@ -109,9 +143,20 @@ void FStatsMallocProfilerProxy::TrackFree( void* Ptr )
 		{
 			FThreadStats* ThreadStats = FThreadStats::GetThreadStats();
 
+#ifdef DEBUG_MALLOC_PROXY
+			if( GThreadStatsToDumpMemory == ThreadStats && ThreadStats->MemoryMessageScope == 0 )
+			{
+				ThreadStats->MemoryMessageScope++;
+				UE_LOG(LogStats, Warning, TEXT("TrackFree, %llu, 0, %i, %i"), (uint64)(UPTRINT)Ptr, SequenceTag, ThreadStats->MemoryMessageScope);
+				ThreadStats->MemoryMessageScope--;
+			}
+#endif // DEBUG_MALLOC_PROXY
+
 			if( ThreadStats->MemoryMessageScope == 0 )
 			{
-				ThreadStats->AddMemoryMessage( GET_STATFNAME(STAT_Memory_FreePtr), (uint64)(UPTRINT)Ptr | (uint64)EMemoryOperation::Free );	// 16 bytes total
+				// 32 bytes per free.
+				ThreadStats->AddMemoryMessage( GET_STATFNAME(STAT_Memory_FreePtr), (uint64)(UPTRINT)Ptr | (uint64)EMemoryOperation::Free );	// 16 bytes total				
+				ThreadStats->AddMemoryMessage( GET_STATFNAME(STAT_Memory_OperationSequenceTag), (int64)SequenceTag );	
 				FreePtrCalls.Increment();
 			}
 		}	
@@ -120,23 +165,28 @@ void FStatsMallocProfilerProxy::TrackFree( void* Ptr )
 
 void* FStatsMallocProfilerProxy::Malloc( SIZE_T Size, uint32 Alignment )
 {
+	const int32 SequenceTag = MemorySequenceTag.Increment();
 	void* Ptr = UsedMalloc->Malloc( Size, Alignment );
 	// We lose the Size's precision, but don't worry about it.
-	TrackAlloc( Ptr, (int64)Size );
+	TrackAlloc( Ptr, (int64)Size, SequenceTag );
 	return Ptr;
 }
 
 void* FStatsMallocProfilerProxy::Realloc( void* OldPtr, SIZE_T NewSize, uint32 Alignment )
 {
-	TrackFree( OldPtr );
+	const int32 FreeSequenceTag = MemorySequenceTag.Increment();
+	TrackFree( OldPtr, FreeSequenceTag );
+
+	const int32 AllocSequenceTag = MemorySequenceTag.Increment();
 	void* NewPtr = UsedMalloc->Realloc( OldPtr, NewSize, Alignment );
-	TrackAlloc( NewPtr, (int64)NewSize );
+	TrackAlloc( NewPtr, (int64)NewSize, AllocSequenceTag );
 	return NewPtr;
 }
 
 void FStatsMallocProfilerProxy::Free( void* Ptr )
 {
-	TrackFree( Ptr );
+	const int32 SequenceTag = MemorySequenceTag.Increment();
+	TrackFree( Ptr, SequenceTag );
 	UsedMalloc->Free( Ptr );
 }
 
diff --git a/Engine/Source/Runtime/Core/Public/Stats/Stats2.h b/Engine/Source/Runtime/Core/Public/Stats/Stats2.h
--- a/Engine/Source/Runtime/Core/Public/Stats/Stats2.h
+++ b/Engine/Source/Runtime/Core/Public/Stats/Stats2.h
@@ -1031,7 +1031,7 @@ typedef TChunkedArray<FStatMessage,(uint32)EStatMessagesArrayConstants::MESSAGES
 */
 struct FStatPacket
 {
-	/** Assigned later, this is the frame number this packet is for **/
+	/** Assigned later, this is the frame number this packet is for. @see FStatsThreadState::ScanForAdvance */
 	int64 Frame;
 	/** ThreadId this packet came from **/
 	uint32 ThreadId;
@@ -1062,6 +1062,30 @@ struct FStatPacket
 		, StatMessagesPresize(Other.StatMessagesPresize)
 	{
 	}
+
+	/** Initializes thread related properties for the stats packet. */
+	void SetThreadProperties()
+	{
+		ThreadId = FPlatformTLS::GetCurrentThreadId();
+		if (ThreadId == GGameThreadId)
+		{
+			ThreadType = EThreadType::Game;
+		}
+		else if (ThreadId == GRenderThreadId)
+		{
+			ThreadType = EThreadType::Renderer;
+		}
+		else
+		{
+			ThreadType = EThreadType::Other;
+		}
+	}
+
+
+	void AssignFrame( int64 InFrame )
+	{
+		Frame = InFrame;
+	}
 };
 
 /** Helper struct used to monitor the scope of the message. */
@@ -1117,54 +1141,59 @@ public:
 	CORE_API void ReturnToPool(FThreadStats* Instance);
 };
 
+/** Fake type to distinguish constructors. */
+enum EConstructor
+{
+	FOR_POOL
+};
+
 /**
 * This is thread-private information about the stats we are acquiring. Pointers to these objects are held in TLS.
 */
-class FThreadStats
+class FThreadStats : FNoncopyable
 {
 	friend class FStatsMallocProfilerProxy;
 	friend class FStatsThreadState;
 	friend class FStatsThread;
 	friend struct FThreadStatsPool;
 
-	enum
-	{
-		PRESIZE_MAX_NUM_ENTRIES = 10,
-		PRESIZE_MAX_SIZE = 256*1024,
-	};
-
 	/** Used to control when we are collecting stats. User of the stats system increment and decrement this counter as they need data. **/
 	CORE_API static FThreadSafeCounter MasterEnableCounter;
 	/** Every time bMasterEnable changes, we update this. This is used to determine frames that have complete data. **/
 	CORE_API static FThreadSafeCounter MasterEnableUpdateNumber;
 	/** while bMasterEnable (or other things affecting stat collection) is chaning, we lock this. This is used to determine frames that have complete data. **/
 	CORE_API static FThreadSafeCounter MasterDisableChangeTagLock;
+	/** TLS slot that holds a FThreadStats. **/
+	CORE_API static uint32 TlsSlot;
 	/** Computed by CheckEnable, the current "master control" for stats collection, based on MasterEnableCounter and a few other things. **/
 	CORE_API static bool bMasterEnable;
 	/** Set to permanently disable the stats system. **/
 	CORE_API static bool bMasterDisableForever;
-	/** TLS slot that holds a FThreadStats. **/
-	CORE_API static uint32 TlsSlot;
+	/** True if we running in the raw stats mode, all stats processing is disabled, captured stats messages are written in timely manner, memory overhead is minimal. */
+	CORE_API static bool bIsRawStatsActive;
 
 	/** The data we are eventually going to send to the stats thread. **/
 	FStatPacket Packet;
 
-	/** Current game frame for this thread stat. */
+	/** Current game frame for this thread stats. */
 	int32 CurrentGameFrame;
 
 	/** Tracks current stack depth for cycle counters. **/
-	uint32 ScopeCount;
+	int32 ScopeCount;
 
 	/** Tracks current stack depth for cycle counters. **/
-	uint32 bWaitForExplicitFlush;
+	int32 bWaitForExplicitFlush;
 
 	/** 
 	 * Tracks the progress of adding a new memory stat message and prevents from using the memory profiler in the scope.
-	 * Mostly related to ignoring all memory allocations in Flush() and AddStatMessage().
+	 * Mostly related to ignoring all memory allocations in AddStatMessage().
 	 * Memory usage of the stats messages is handled by the STAT_StatMessagesMemory.
 	 */
 	int32 MemoryMessageScope;
 
+	/** We have a scoped cycle counter in the FlushRawStats which causes an infinite recursion, ReentranceGuard will solve that issue. */
+	bool bReentranceGuard;
+
 	/** Tracks current stack depth for cycle counters. **/
 	bool bSawExplicitFlush;
 
@@ -1172,11 +1201,9 @@ class FThreadStats
 	CORE_API FThreadStats();
 
 	/** Constructor used for the pool. */
-	CORE_API FThreadStats(bool);
-
-	/** Copys a packet for sending. !!!CAUTION!!! does not copy packets.**/
-	CORE_API FThreadStats(FThreadStats const& Other);
+	CORE_API FThreadStats(EConstructor);
 
+public:
 	/** Checks the TLS for a thread packet and if it isn't found, it makes a new one. **/
 	static FORCEINLINE_STATS FThreadStats* GetThreadStats()
 	{
@@ -1191,8 +1218,49 @@ class FThreadStats
 	/** This should be called when conditions have changed such that stat collection may now be enabled or not **/
 	static CORE_API void CheckEnable();
 
-	/** Send any outstanding packets to the stats thread **/
-	CORE_API void Flush(bool bHasBrokenCallstacks = false);
+	/**
+	 *	Checks if the game frame has changed and updates the current game frame.
+	 *	Used by the flushing mechanism to optimize the memory usage, only valid for packets from other threads.
+	 *	@return true, if the frame has changed
+	 */
+	FORCEINLINE_STATS bool DetectAndUpdateCurrentGameFrame()
+	{
+		if( Packet.ThreadType == EThreadType::Other )
+		{
+			FPlatformMisc::MemoryBarrier();
+			const bool bFrameHasChanged = FStats::GameThreadStatsFrame > CurrentGameFrame;
+			if( bFrameHasChanged )
+			{
+				CurrentGameFrame = FStats::GameThreadStatsFrame;
+				return true;
+			}
+		}
+				
+		return false;
+	}
+
+	/** Maintains the explicit flush. */
+	FORCEINLINE_STATS void UpdateExplicitFlush()
+	{
+		if (Packet.ThreadType != EThreadType::Other && bSawExplicitFlush)
+		{
+			bSawExplicitFlush = false;
+			bWaitForExplicitFlush = 1;
+			ScopeCount++; // prevent sends until the next explicit flush
+		}
+	}
+
+	/** Send any outstanding packets to the stats thread. **/
+	CORE_API void Flush(bool bHasBrokenCallstacks = false, bool bForceFlush = false);
+
+	/** Flushes the regular stats, the realtime stats. */
+	CORE_API void FlushRegularStats(bool bHasBrokenCallstacks, bool bForceFlush);
+
+	/** Flushes the raw stats, low memory and performance overhead, but not realtime. */
+	CORE_API void FlushRawStats(bool bHasBrokenCallstacks = false, bool bForceFlush = false);
+
+	/** Checks the command line whether we want to enable collecting startup stats. */
+	static CORE_API void CheckForCollectingStartupStats();
 
 	FORCEINLINE_STATS void AddStatMessage( const FStatMessage& StatMessage )
 	{
@@ -1207,7 +1275,8 @@ public:
 		FThreadStats* Stats = (FThreadStats*)FPlatformTLS::GetTlsValue(TlsSlot);
 		if (Stats)
 		{
-			// Flush stats?
+			// Send all remaining messages.
+			Stats->Flush(false, true);
 			FPlatformTLS::SetTlsValue(TlsSlot, nullptr);
 			FThreadStatsPool::Get().ReturnToPool(Stats);
 		}
@@ -1223,6 +1292,11 @@ public:
 		{
 			ThreadStats->ScopeCount++;
 			ThreadStats->AddStatMessage(FStatMessage(InStatName, InStatOperation));
+
+			if( bIsRawStatsActive )
+			{
+				ThreadStats->FlushRawStats();
+			}
 		}
 		else if (InStatOperation == EStatOperation::CycleScopeEnd)
 		{
@@ -1234,6 +1308,10 @@ public:
 				{
 					ThreadStats->Flush();
 				}
+				else if( bIsRawStatsActive )
+				{
+					ThreadStats->FlushRawStats();
+				}
 			}
 			// else we dumped this frame without closing scope, so we just drop the closes on the floor
 		}
@@ -1251,6 +1329,10 @@ public:
 			{
 				ThreadStats->Flush();
 			}
+			else if( bIsRawStatsActive )
+			{
+				ThreadStats->FlushRawStats();
+			}
 		}
 	}
 
@@ -1347,6 +1429,20 @@ public:
 		FPlatformMisc::MemoryBarrier();
 	}
 
+	/** Enables the raw stats mode. */
+	static FORCEINLINE_STATS void EnableRawStats()
+	{
+		bIsRawStatsActive = true;
+		FPlatformMisc::MemoryBarrier();
+	}
+
+	/** Disables the raw stats mode. */
+	static FORCEINLINE_STATS void DisableRawStats()
+	{
+		bIsRawStatsActive = false;
+		FPlatformMisc::MemoryBarrier();
+	}
+
 	/** Called by launch engine loop to start the stats thread **/
 	static CORE_API void StartThread();
 	/** Called by launch engine loop to stop the stats thread **/
@@ -1875,11 +1971,6 @@ DECLARE_STATS_GROUP(TEXT("UObjects"),STATGROUP_UObjects, STATCAT_Advanced);
 DECLARE_STATS_GROUP(TEXT("User"),STATGROUP_User, STATCAT_Advanced);
 
 DECLARE_CYCLE_STAT_EXTERN(TEXT("FrameTime"),STAT_FrameTime,STATGROUP_Engine, CORE_API);
-DECLARE_FLOAT_COUNTER_STAT_EXTERN(TEXT("StatUnit FPS"), STAT_FPS, STATGROUP_Engine, CORE_API);
-
-/** Stats for the stat system */
 
-DECLARE_CYCLE_STAT_EXTERN(TEXT("DrawStats"),STAT_DrawStats,STATGROUP_StatSystem, CORE_API);
-DECLARE_MEMORY_STAT_EXTERN(TEXT("Stat Messages"),STAT_StatMessagesMemory,STATGROUP_StatSystem, CORE_API);
 
 #endif
diff --git a/Engine/Source/Runtime/Core/Public/Stats/StatsData.h b/Engine/Source/Runtime/Core/Public/Stats/StatsData.h
--- a/Engine/Source/Runtime/Core/Public/Stats/StatsData.h
+++ b/Engine/Source/Runtime/Core/Public/Stats/StatsData.h
@@ -192,10 +192,15 @@ struct CORE_API FStatPacketArray
 		Empty();
 	}
 
+	/** Deletes all stats packets. */
 	void Empty();
+
+	void RemovePtrsButNoData()
+	{
+		Packets.Empty();
+	}
 };
 
-// @TODO yrx 2014-11-21 basic memory profiling?
 /**
 * A call stack of stat messages. Normally we don't store or transmit these; they are used for stats processing and visualization.
 */
@@ -391,9 +396,6 @@ protected:
 
 //@todo split header
 
-/** Delegate that FStatsThreadState calls on the stats thread whenever we have a new frame. */
-DECLARE_MULTICAST_DELEGATE_OneParam(FOnNewFrameHistory, int64);
-
 /**
 * Some of the FStatsThreadState data methods allow filtering. Derive your filter from this class
 */
@@ -403,37 +405,9 @@ struct IItemFiler
 	virtual bool Keep(FStatMessage const& Item) = 0;
 };
 
-/** Simple allocation info, temporary. */
-struct FAllocationInfo
-{
-	int64 Size;
-	FName Scope;
-	FAllocationInfo(int64 InSize, FName InScope)
-		: Size(InSize)
-		, Scope(InScope)
-	{}
-	FAllocationInfo()
-		: Size(0)
-	{}
-};
 
-/** Simple allocation info, temporary. */
-struct FAllocationInfoEx
-{
-	int64 Size;
-	FName Scope;
-	int64 Cycles;
-	FAllocationInfoEx(int64 InSize, FName InScope, int64 InCycles)
-		: Size(InSize)
-		, Scope(InScope)
-		, Cycles(InCycles)
-	{}
-	FAllocationInfoEx()
-		: Size(0)
-		, Cycles(0)
-	{}
-};
 
+// @TODO yrx 2014-12-03 Separete stats thread state vs rawstats thread state?
 // @TODO yrx 2014-03-21 Move metadata functionality into a separate class
 /**
 * Tracks stat state and history
@@ -443,13 +417,21 @@ struct FAllocationInfoEx
 class CORE_API FStatsThreadState
 {
 	friend class FStatsThread;
+	friend struct FStatPacketArray;
+
+	/** Delegate that FStatsThreadState calls on the stats thread whenever we have a new frame. */
+	DECLARE_MULTICAST_DELEGATE_OneParam(FOnNewFrameHistory, int64);
 
-	/** Internal method to scan the messages to update the current frame **/
+	/** Delegate that FStatsThreadState calls on the stats thread whenever we have a new raw stats packet. */
+	DECLARE_MULTICAST_DELEGATE_OneParam(FOnNewRawStatPacket, const FStatPacket*);
+
+	/** Internal method to scan the messages to update the current frame. */
 	void ScanForAdvance(const FStatMessagesArray& Data);
 
-	/** Internal method to scan the messages to update the current frame and accumulate any non-frame stats. **/
+	/** Internal method to scan the messages to update the current frame. */
 	void ScanForAdvance(FStatPacketArray& NewData);
 
+	// @TODO yrx 2014-11-28 Remove from here
 	/** Internal method to scan the messages to accumulate any non-frame stats. **/
 	void ProcessMetaDataForLoad(TArray<FStatMessage>& Data);
 
@@ -457,6 +439,7 @@ public:
 	/** Internal method to add meta data packets to the data structures. **/
 	void ProcessMetaDataOnly(TArray<FStatMessage>& Data);
 
+	// @TODO yrx 2014-11-28 Remove from here
 	/** Marks this stats state as loaded. */
 	void MarkAsLoaded()
 	{
@@ -466,13 +449,31 @@ public:
 	/** Toggles tracking the most memory expensive stats. */
 	void ToggleFindMemoryExtensiveStats();
 
+	void ResetStatsForRawStats()
+	{
+		MaxNumStatMessages = 0;
+		TotalNumStatMessages = 0;
+	}
+
 private:
-	/** Internal method to scan the messages to accumulate any non-frame stats. **/
+	/** Internal method to accumulate any non-frame stats. **/
 	void ProcessNonFrameStats(FStatMessagesArray& Data, TSet<FName>* NonFrameStatsFound);
 
 	/** Internal method to place the data into the history, discard and broadcast any new frames to anyone who cares. **/
 	void AddToHistoryAndEmpty(FStatPacketArray& NewData);
 
+	/** Does basic processing on the raw stats packets, discard and broadcast any new raw stats packets to anyone who cares. */
+	void ProcessRawStats(FStatPacketArray& NewData);
+
+	/** Resets the state of the raw stats. */
+	void ResetRawStats();
+
+	/** Reset the state of the regular stats. */
+	void ResetRegularStats();
+
+	/** Prepares fake FGameThreadHudData to display the raw stats memory overhead. */
+	void UpdateStatMessagesMemoryUsage();
+
 	/** Generates a list of most memory expensive stats and dump to the log. */
 	void FindAndDumpMemoryExtensiveStats( FStatPacketArray &Frame );
 
@@ -491,9 +492,11 @@ private:
 	/** Number of frames to keep in the history. **/
 	int32 HistoryFrames;
 
+	// @TODO yrx 2014-11-28 
 	/** Largest frame seen. Loaded stats only. **/
 	int64 MaxFrameSeen;
 
+	// @TODO yrx 2014-11-28 
 	/** First frame seen. Loaded stats only. **/
 	int64 MinFrameSeen;
 
@@ -503,39 +506,49 @@ private:
 	/** Used to track which packets have been sent to listeners **/
 	int64 LastFullFrameProcessed;
 
-	/** Valid frame computation is different if we just loaded these stats **/
-	bool bWasLoaded;
-
-	/** If true, stats each frame will dump to the log a list of most memory expensive stats. */
-	bool bFindMemoryExtensiveStats;
-
 	/** Cached condensed frames. This saves a lot of time since multiple listeners can use the same condensed data **/
 	mutable TMap<int64, TArray<FStatMessage>* > CondensedStackHistory;
 
-	/** Map from allocation pointer to allocation info, currently active allocations */
-	mutable TMap< uint64, FAllocationInfo > AllocationMap;
-
-	/** Map from allocation pointer to allocation info. */
-	mutable TMap< uint64, FAllocationInfoEx > NonSequentialAllocationMap;
-
 	/** After we add to the history, these frames are known good. **/
 	TSet<int64>	GoodFrames;
 
 	/** These frames are known bad. **/
 	TSet<int64>	BadFrames;
 
+	/** Array of stat packets that are not processed yet. */
+	FStatPacketArray StartupRawStats;
+
+	/** Total number of stats messages for active session. */
+	int64 TotalNumStatMessages;
+
+	/** Current number of stat messages stored by stats thread and ready to be processed (for visualization, for saving). */
+	FThreadSafeCounter NumStatMessages;
+
+	/** Maximum number of stat messages seen so far for active session. */
+	int32 MaxNumStatMessages;
+
+	// @TODO yrx 2014-11-28 
+	/** Valid frame computation is different if we just loaded these stats **/
+	bool bWasLoaded;
+
+	/** If true, stats each frame will dump to the log a list of most memory expensive stats. */
+	bool bFindMemoryExtensiveStats;
+
 public:
 
 	/** Constructor used by GetLocalState(), also used by the profiler to hold a previewing stats thread state. We don't keep many frames by default **/
 	FStatsThreadState(int32 InHistoryFrames = STAT_FRAME_SLOP + 2);
 
+	// @TODO yrx 2014-11-28 
 	/** Constructor to load stats from a file **/
 	FStatsThreadState(FString const& Filename);
 
 	/** Delegate we fire every time we have a new complete frame of data. **/
-	// @TODO yrx 2014-03-21 Not thread-safe ?
 	mutable FOnNewFrameHistory NewFrameDelegate;
 
+	/** Delegate we fire every time we have a new raw stats packet, the data is no longer valid outside the broadcast. */
+	mutable FOnNewRawStatPacket NewRawStatPacket;
+
 	/** New packets not on the render thread as assign this frame **/
 	int64 CurrentGameFrame;
 
@@ -572,17 +585,6 @@ public:
 		return GoodFrames.Contains(Frame);
 	}
 
-	/** Gets allocations map. */
-	const TMap< uint64, FAllocationInfo >& GetAllocations() const 
-	{ 
-		return AllocationMap; 
-	}
-
-	const TMap< uint64, FAllocationInfoEx >& GetNonSequentialAllocations() const 
-	{ 
-		return NonSequentialAllocationMap; 
-	}
-
 	/** Used by the hitch detector. Does a very fast look a thread to decide if this is a hitch frame. **/
 	int64 GetFastThreadFrameTime(int64 TargetFrame, EThreadType::Type Thread) const;
 	int64 GetFastThreadFrameTime(int64 TargetFrame, uint32 ThreadID) const;
@@ -613,6 +615,7 @@ public:
 	/** Adds missing stats to the group so it doesn't jitter. **/
 	void AddMissingStats(TArray<FStatMessage>& Dest, TSet<FName> const& EnabledItems) const;
 
+	// @TODO yrx 2014-11-28 REmove from here, it's file related.
 	/** Adds a frame worth of messages */
 	void AddMessages(TArray<FStatMessage>& InMessages);
 
@@ -820,11 +823,19 @@ struct FHudGroup
 */
 struct FGameThreadHudData
 {
+	/** Initialization constructor. */
+	FGameThreadHudData( bool bInDrawOnlyRawStats )
+		: bDrawOnlyRawStats(bInDrawOnlyRawStats)
+	{}
+
 	TIndirectArray<FHudGroup> HudGroups;
 	TArray<FName> GroupNames;
 	TArray<FString> GroupDescriptions;
 	TMap<FPlatformMemory::EMemoryCounterRegion, int64> PoolCapacity;
 	TMap<FPlatformMemory::EMemoryCounterRegion, FString> PoolAbbreviation;
+
+	/** Whether to display minimal stats for the raw stats mode. */
+	const bool bDrawOnlyRawStats;
 };
 
 /**
diff --git a/Engine/Source/Runtime/Core/Public/Stats/StatsFile.h b/Engine/Source/Runtime/Core/Public/Stats/StatsFile.h
--- a/Engine/Source/Runtime/Core/Public/Stats/StatsFile.h
+++ b/Engine/Source/Runtime/Core/Public/Stats/StatsFile.h
@@ -5,8 +5,7 @@
 
 #if	STATS
 
-class FAsyncWriteWorker;
-
+class FAsyncStatsWrite;
 
 /**
 * Magic numbers for stats streams, this is for the first version.
@@ -36,9 +35,16 @@ namespace EStatMagicWithHeader
 
 		/**
 		 *	Added support for compressing the stats data, now each frame is compressed.
+		 *	!!CAUTION!! Deprecated.
 		 */
 		VERSION_4 = 4,
-		HAS_COMPRESSED_DATA_VER = VERSION_4
+		HAS_COMPRESSED_DATA_VER = VERSION_4,
+
+		/**
+		 *	New low-level raw stats with memory profiler functionality.
+		 *  !!CAUTION!! Not backward compatible with version 4.
+		 */
+		VERSION_5 = 5
 	};
 }
 
@@ -48,7 +54,6 @@ struct EStatsFileConstants
 	{
 		/**
 		 *	Maximum size of the data that can be compressed.
-		 *	Should be ok for almost all cases, except the boot stats.
 		 */
 		MAX_COMPRESSED_SIZE = 1024 * 1024,
 
@@ -68,7 +73,7 @@ struct EStatsFileConstants
 -----------------------------------------------------------------------------*/
 
 /** Helper struct used to operate on the compressed data. */
-struct FCompressedStatsData
+struct CORE_API FCompressedStatsData
 {
 	/**
 	 * Initialization constructor 
@@ -204,15 +209,17 @@ struct FStatsStreamHeader
 	{}
 
 	/**
-	*	Version number to detect version mismatches
-	*	1 - Initial version for supporting raw ue4 stats files
-	*/
+	 *	Version number to detect version mismatches
+	 *	1 - Initial version for supporting raw ue4 stats files
+	 */
 	uint32	Version;
 
 	/** Platform that this file was captured on. */
 	FString	PlatformName;
 
-	/** Offset in the file for the frame table. Serialized as TArray<int64>. */
+	/** 
+	 *  Offset in the file for the frame table. Serialized as TArray<int64>
+	 *  For raw stats contains only one element which indicates the begin of the data. */
 	uint64	FrameTableOffset;
 
 	/** Offset in the file for the FName table. Serialized with WriteFName/ReadFName. */
@@ -227,10 +234,10 @@ struct FStatsStreamHeader
 	/** Number of the metadata messages. */
 	uint64	NumMetadataMessages;
 
-	/** Whether this stats file uses raw data, required for thread view. */
+	/** Whether this stats file uses raw data, required for thread view/memory profiling/advanced profiling. */
 	bool bRawStatsFile;
 
-	/** Whether this stats file has been finalized, so we have the whole data and can load file a unordered access. */
+	/** Whether this stats file has all names stored at the end of file. */
 	bool IsFinalized() const
 	{
 		return NumMetadataMessages > 0 && MetadataMessagesOffset > 0 && FrameTableOffset > 0;
@@ -281,6 +288,11 @@ struct FStatsFrameInfo
 		: FrameFileOffset(0)
 	{}
 
+	/** Initialization constructor. */
+	FStatsFrameInfo( int64 InFrameFileOffset )
+		: FrameFileOffset(InFrameFileOffset)
+	{}
+
 	/** Initialization constructor. */
 	FStatsFrameInfo( int64 InFrameFileOffset, TMap<uint32, int64>& InThreadCycles )
 		: FrameFileOffset(InFrameFileOffset)
@@ -301,58 +313,68 @@ struct FStatsFrameInfo
 	TMap<uint32, int64> ThreadCycles;
 };
 
-
-struct CORE_API FStatsWriteFile : public TSharedFromThis<FStatsWriteFile, ESPMode::ThreadSafe>
+/** Interface for writing stats data. Can be only used in the stats thread. */
+struct CORE_API IStatsWriteFile
 {
-	friend class FAsyncWriteWorker;
+	friend class FAsyncStatsWrite;
 
 protected:
+	/** Stats file archive. */
+	FArchive* File;
+
+	/** Filename of the archive that we are writing to. */
+	FString ArchiveFilename;
+
 	/** Stats stream header. */
 	FStatsStreamHeader Header;
 
-	/** Set of names already sent **/
+	/** Set of names already sent. */
 	TSet<int32> FNamesSent;
 
-	/** Array of stats frames info already captured. */
-	TArray<FStatsFrameInfo> FramesInfo;
+	/** Async task used to offload saving the capture data. */
+	FAsyncTask<FAsyncStatsWrite>* AsyncTask;
 
 	/** Data to write through the async task. **/
 	TArray<uint8> OutData;
 
-	/** Thread cycles for the last frame. */
-	TMap<uint32, int64> ThreadCycles;
-
-	/** Filename of the archive that we are writing to. */
-	FString ArchiveFilename;
-
-	/** Stats file archive. */
-	FArchive* File;
+	/** Buffer to store the compressed data, used by the FAsyncStatsWrite. */
+	TArray<uint8> CompressedData;
 
-	/** Async task used to offload saving the capture data. */
-	FAsyncTask<FAsyncWriteWorker>* AsyncTask;
+	/**
+	 *  Array of stats frames info already captured.
+	 *  !!CAUTION!!
+	 *  Only modified in the async write thread through FinalizeSavingData method.
+	 */
+	TArray<FStatsFrameInfo> FramesInfo;
 
-	/** Buffer to store the compressed data, used by the FAsyncWriteWorker. */
-	TArray<uint8> CompressedData;
+protected:
+	/** Default constructor. */
+	IStatsWriteFile();
 
 	/** NewFrame delegate handle  */
-	FDelegateHandle NewFrameDelegateHandle;
+	FDelegateHandle DataDelegateHandle;
 
 public:
-	/** Constructor. **/
-	FStatsWriteFile();
+	/** Creates a file writer and registers for the data delegate. */
+	void Start( const FString& InFilename );
 
-	/** Virtual destructor. */
-	virtual ~FStatsWriteFile()
-	{}
+	/** Finalizes writing the stats data and unregisters the data delegate. */
+	void Stop();
 
-	/** Starts writing the stats data into the specified file. */
-	void Start( FString const& InFilename, bool bIsRawStatsFile );
+protected:
+	/** Sets the data delegate used to receive a stats data. */
+	virtual void SetDataDelegate( bool bSet ) = 0;
 
-	/** Stops writing the stats data. */
-	void Stop();
+	/** Finalization code called after the compressed data has been saved. */
+	virtual void FinalizeSavingData( int64 FrameFileOffset )
+	{};
 
-	bool IsValid() const;
+	bool IsValid() const
+	{
+		return !!File;
+	}
 
+public:
 	const TArray<uint8>& GetOutData() const
 	{
 		return OutData;
@@ -363,47 +385,19 @@ public:
 		OutData.Reset();
 	}
 
-	/**
-	*	Writes magic value, dummy header and initial metadata.
-	*	Called from the stats thread.
-	*/
-	void WriteHeader( bool bIsRawStatsFile );
-
-	/**
-	 *	Grabs a frame from the local FStatsThreadState and adds it to the output.
-	 *	Called from the stats thread, but the data is saved using the the FAsyncWriteWorker. 
-	 */
-	virtual void WriteFrame( int64 TargetFrame, bool bNeedFullMetadata = false );
+	/** Writes magic value, dummy header and initial metadata. */
+	void WriteHeader();
 
 protected:
-	/**
-	 *	Finalizes writing to the file.
-	 *	Called from the stats thread.
-	 */
-	void Finalize();
+	/** Writes metadata messages into the stream. */
+	void WriteMetadata( FArchive& Ar );
 
-	void NewFrame( int64 TargetFrame );
+	/**	Finalizes writing to the file. */
+	void Finalize();
 
 	/** Sends the data to the file via async task. */
 	void SendTask();
 
-	void AddNewFrameDelegate();
-	void RemoveNewFrameDelegate();
-
-	/**
-	 *	Writes metadata messages into the stream. 
-	 *	Called from the stats thread.
-	 */
-	void WriteMetadata( FArchive& Ar )
-	{
-		const FStatsThreadState& Stats = FStatsThreadState::GetLocalState();
-		for( const auto& It : Stats.ShortNameToLongName )
-		{
-			const FStatMessage& StatMessage = It.Value;
-			WriteMessage( Ar, StatMessage );
-		}
-	}
-
 	/** Sends an FName, and the string it represents if we have not sent that string before. **/
 	FORCEINLINE_STATS void WriteFName( FArchive& Ar, FStatNameAndInfo NameAndInfo )
 	{
@@ -435,54 +429,92 @@ protected:
 			{
 				int64 Payload = Item.GetValue_int64();
 				Ar << Payload;
-			}
 				break;
+			}
+				
 			case EStatDataType::ST_double:
 			{
 				double Payload = Item.GetValue_double();
 				Ar << Payload;
-			}
 				break;
+			}
+
 			case EStatDataType::ST_FName:
+			{
 				WriteFName( Ar, FStatNameAndInfo( Item.GetValue_FName(), false ) );
 				break;
+			}
+
+			case EStatDataType::ST_Ptr:
+			{
+				uint64 Payload = Item.GetValue_Ptr();
+				Ar << Payload;
+				break;
+			}
 		}
 	}
 };
 
-/** Struct used to save unprocessed stats data as raw messages. */
-struct FRawStatsWriteFile : public FStatsWriteFile
+/** Helper struct used to write regular stats to the file. */
+struct CORE_API FStatsWriteFile : public IStatsWriteFile
 {
+	friend class FAsyncStatsWrite;
+
+protected:
+	/** Thread cycles for the last frame. */
+	TMap<uint32, int64> ThreadCycles;
+
 public:
-	/** Virtual destructor. */
-	virtual ~FRawStatsWriteFile()
-	{}
+	/** Default constructor, set bRawStatsFile to false. */
+	FStatsWriteFile()
+	{
+		Header.bRawStatsFile = false;
+	}
 
 protected:
+	virtual void SetDataDelegate( bool bSet ) override;
+
+	virtual void FinalizeSavingData( int64 FrameFileOffset ) override;
 
+public:
 	/**
-	*	Grabs a frame from the local FStatsThreadState and adds it to the output.
-	*	Called from the stats thread, but the data is saved using the the FAsyncWriteWorker.
-	*/
-	virtual void WriteFrame( int64 TargetFrame, bool bNeedFullMetadata = false ) override;
+	 *	Grabs a frame from the local FStatsThreadState and writes it to the array.
+	 *	Only used by the ProfilerServiceManager.
+	 */
+	void WriteFrame( int64 TargetFrame, bool bNeedFullMetadata );
 
-	/** Write a stat packed into the specified archive. */
-	void WriteStatPacket( FArchive& Ar, /*const*/ FStatPacket& StatPacket )
+protected:
+	/**
+	 *	Grabs a frame from the local FStatsThreadState and adds it to the output.
+	 *	Called from the stats thread, but the data is saved using the the FAsyncStatsWrite. 
+	 */
+	void WriteFrame( int64 TargetFrame )
 	{
-		Ar << StatPacket.Frame;
-		Ar << StatPacket.ThreadId;
-		int32 MyThreadType = (int32)StatPacket.ThreadType;
-		Ar << MyThreadType;
+		WriteFrame( TargetFrame, false );
+		SendTask();
+	}
+};
 
-		Ar << StatPacket.bBrokenCallstacks;
-		// We must handle stat messages in a different way.
-		int32 NumMessages = StatPacket.StatMessages.Num();
-		Ar << NumMessages;
-		for( int32 MessageIndex = 0; MessageIndex < NumMessages; ++MessageIndex )
-		{
-			WriteMessage( Ar, StatPacket.StatMessages[MessageIndex] );
-		}
+/** Helper struct used to write raw stats to the file. */
+struct CORE_API FRawStatsWriteFile : public IStatsWriteFile
+{
+	bool bWrittenOffsetToData;
+
+public:
+	/** Default constructor, set bRawStatsFile to true. */
+	FRawStatsWriteFile()
+		: bWrittenOffsetToData(false)
+	{
+		Header.bRawStatsFile = true;
 	}
+
+protected:
+	virtual void SetDataDelegate( bool bSet ) override;
+
+	void WriteRawStatPacket( const FStatPacket* StatPacket );
+
+	/** Write a stat packed into the specified archive. */
+	void WriteStatPacket( FArchive& Ar, FStatPacket& StatPacket );
 };
 
 /*-----------------------------------------------------------------------------
@@ -490,8 +522,8 @@ protected:
 -----------------------------------------------------------------------------*/
 
 /**
-* Class for maintaining state of receiving a stream of stat messages
-*/
+ * Class for maintaining state of receiving a stream of stat messages
+ */
 struct CORE_API FStatsReadStream
 {
 public:
@@ -501,7 +533,7 @@ public:
 	/** FNames have a different index on each machine, so we translate via this map. **/
 	TMap<int32, int32> FNamesIndexMap;
 
-	/** Array of stats frame info. */
+	/** Array of stats frame info. Empty for the raw stats. */
 	TArray<FStatsFrameInfo> FramesInfo;
 
 	/** Reads a stats stream header, returns true if the header is valid and we can continue reading. */
@@ -542,8 +574,10 @@ public:
 		return true;
 	}
 
+	// @TODO yrx 2014-12-02 Add a better way to read the file.
+
 	/** Reads a stat packed from the specified archive. */
-	void ReadStatPacket( FArchive& Ar, FStatPacket& StatPacked, bool bHasFNameMap )
+	void ReadStatPacket( FArchive& Ar, FStatPacket& StatPacked )
 	{
 		Ar << StatPacked.Frame;
 		Ar << StatPacked.ThreadId;
@@ -558,7 +592,7 @@ public:
 		StatPacked.StatMessages.Reserve( NumMessages );
 		for( int32 MessageIndex = 0; MessageIndex < NumMessages; ++MessageIndex )
 		{
-			new(StatPacked.StatMessages) FStatMessage( ReadMessage( Ar, bHasFNameMap ) );
+			new(StatPacked.StatMessages) FStatMessage( ReadMessage( Ar, true ) );
 		}
 	}
 
@@ -638,19 +672,31 @@ public:
 				int64 Payload = 0;
 				Ar << Payload;
 				Result.GetValue_int64() = Payload;
-			}
 				break;
+			}
+				
 			case EStatDataType::ST_double:
 			{
 				double Payload = 0;
 				Ar << Payload;
 				Result.GetValue_double() = Payload;
-			}
 				break;
+			}
+				
 			case EStatDataType::ST_FName:
+			{
 				FStatNameAndInfo Payload( ReadFName( Ar, bHasFNameMap ) );
 				Result.GetValue_FName() = Payload.GetRawName();
 				break;
+			}
+
+			case EStatDataType::ST_Ptr:
+			{
+				uint64 Payload = 0;
+				Ar << Payload;
+				Result.GetValue_Ptr() = Payload;
+				break;
+			}
 		}
 		return Result;
 	}
@@ -689,6 +735,7 @@ public:
 	}
 };
 
+
 /*-----------------------------------------------------------------------------
 	Commands functionality
 -----------------------------------------------------------------------------*/
@@ -696,23 +743,41 @@ public:
 /** Implements 'Stat Start/StopFile' functionality. */
 struct FCommandStatsFile
 {
-	/** Filename of the last saved stats file. */
-	static FString LastFileSaved;
+	/** Access to the singleton. */
+	static CORE_API FCommandStatsFile& Get();
 
-	/** First frame. */
-	static int64 FirstFrame;
-
-	/** Stats file that is currently being saved. */
-	static TSharedPtr<FStatsWriteFile, ESPMode::ThreadSafe> CurrentStatsFile;
+	/** Default constructor. */
+	FCommandStatsFile()
+		: FirstFrame(-1)
+		, CurrentStatsFile(nullptr)
+	{}
 
 	/** Stat StartFile. */
-	static void Start( const TCHAR* Cmd );
+	void Start( const FString& File );
 
 	/** Stat StartFileRaw. */
-	static void StartRaw( const TCHAR* Cmd );
+	void StartRaw( const FString& File );
 
 	/** Stat StopFile. */
-	static void Stop();
+	void Stop();
+
+	bool IsStatFileActive()
+	{
+		return StatFileActiveCounter.GetValue() > 0;
+	}
+
+	/** Filename of the last saved stats file. */
+	FString LastFileSaved;
+
+protected:
+	/** First frame. */
+	int64 FirstFrame;
+
+	/** Whether the 'stat startfile' command is active, can be accessed by other thread. */
+	FThreadSafeCounter StatFileActiveCounter;
+
+	/** Stats file that is currently being saved. */
+	IStatsWriteFile* CurrentStatsFile;
 };
 
 
diff --git a/Engine/Source/Runtime/Core/Public/Stats/StatsMallocProfilerProxy.h b/Engine/Source/Runtime/Core/Public/Stats/StatsMallocProfilerProxy.h
--- a/Engine/Source/Runtime/Core/Public/Stats/StatsMallocProfilerProxy.h
+++ b/Engine/Source/Runtime/Core/Public/Stats/StatsMallocProfilerProxy.h
@@ -15,8 +15,17 @@ private:
 	
 	/** Whether the stats malloc profiler is enabled, disabled by default */
 	bool bEnabled;
+	
+	/** True if the stats malloc profiler was enabled, we can't enable malloc profiler second time. */
+	bool bWasEnabled;
+
+	/** Memory sequence tag, used by the processing tool to indicate the order of allocations. */
+	FThreadSafeCounter MemorySequenceTag;
 
+	/** Malloc calls made this frame. */
 	FThreadSafeCounter AllocPtrCalls;
+
+	/** Free calls made this frame. */
 	FThreadSafeCounter FreePtrCalls;
 
 public:
@@ -29,7 +38,7 @@ public:
 
 	static CORE_API FStatsMallocProfilerProxy* Get();
 
-	void SetState( bool bEnable );
+	void SetState( bool bNewState );
 
 	bool GetState() const
 	{
@@ -41,17 +50,19 @@ public:
 	/**
 	 * Tracks malloc operation.
 	 *
-	 * @param	Ptr	- Allocated pointer 
-	 * @param	Size- Size of allocated pointer
+	 * @param	Ptr			- Allocated pointer 
+	 * @param	Size		- Size of allocated pointer
+	 * @param	SequenceTag - Previous value for the memory sequence tag
 	 */
-	void TrackAlloc( void* Ptr, int64 Size );
+	void TrackAlloc( void* Ptr, int64 Size, int32 SequenceTag );
 
 	/**
 	 * Tracks free operation
 	 *
 	 * @param	Ptr	- Freed pointer
+	 * @param	SequenceTag - Previous value for the memory sequence tag
 	 */
-	void TrackFree( void* Ptr );
+	void TrackFree( void* Ptr, int32 SequenceTag );
 
 	virtual void* Malloc( SIZE_T Size, uint32 Alignment ) override;
 	virtual void* Realloc( void* OldPtr, SIZE_T NewSize, uint32 Alignment ) override;
diff --git a/Engine/Source/Runtime/Engine/Private/StatsRender2.cpp b/Engine/Source/Runtime/Engine/Private/StatsRender2.cpp
--- a/Engine/Source/Runtime/Engine/Private/StatsRender2.cpp
+++ b/Engine/Source/Runtime/Engine/Private/StatsRender2.cpp
@@ -669,7 +669,19 @@ void RenderStats(FViewport* Viewport, class FCanvas* Canvas, int32 X, int32 Y)
 		return;
 	}
 	GetStatRenderGlobals().Initialize( Viewport->GetSizeXY() );
-	RenderGroupedWithHierarchy(*ViewData, Viewport, Canvas, X, Y);
+	if( !ViewData->bDrawOnlyRawStats )
+	{
+		RenderGroupedWithHierarchy(*ViewData, Viewport, Canvas, X, Y);
+	}
+	else
+	{
+		// Render all counters.
+		for( int32 RowIndex = 0; RowIndex < ViewData->GroupDescriptions.Num(); ++RowIndex )
+		{
+			Canvas->DrawShadowedString(X, Y, *ViewData->GroupDescriptions[RowIndex], GetStatRenderGlobals().StatFont, GetStatRenderGlobals().StatColor);
+			Y += GetStatRenderGlobals().GetFontHeight();
+		}
+	}
 }
 
 #endif
diff --git a/Engine/Source/Runtime/Engine/Private/UnrealEngine.cpp b/Engine/Source/Runtime/Engine/Private/UnrealEngine.cpp
--- a/Engine/Source/Runtime/Engine/Private/UnrealEngine.cpp
+++ b/Engine/Source/Runtime/Engine/Private/UnrealEngine.cpp
@@ -106,6 +106,8 @@ IMPLEMENT_MODULE( FEngineModule, Engine );
 
 #define LOCTEXT_NAMESPACE "UnrealEngine"
 
+DECLARE_CYCLE_STAT(TEXT("DrawStats"),STAT_DrawStats,STATGROUP_StatSystem);
+
 void FEngineModule::StartupModule()
 {
 	// Setup delegate callback for ProfilingHelpers to access current map name
@@ -8680,8 +8682,6 @@ bool UEngine::LoadMap( FWorldContext& WorldContext, FURL URL, class UPendingNetG
 	WorldContext.SeamlessTravelHandler.CancelTravel();
 
 	double	StartTime = FPlatformTime::Seconds();
-	{
-		DECLARE_SCOPE_CYCLE_COUNTER(TEXT("Loading URL"), STAT_LoadMap, STATGROUP_LoadTime);
 
 	UE_LOG(LogLoad, Log,  TEXT("LoadMap: %s"), *URL.ToString() );
 	GInitRunaway();
@@ -9109,7 +9109,6 @@ bool UEngine::LoadMap( FWorldContext& WorldContext, FURL URL, class UPendingNetG
 	IStreamingManager::Get().NotifyLevelChange();
 
 	WorldContext.World()->BeginPlay();
-	}
 
 	// send a callback message
 	PostLoadMapCaller.bCalled = true;
diff --git a/Engine/Source/Runtime/Launch/Private/LaunchEngineLoop.cpp b/Engine/Source/Runtime/Launch/Private/LaunchEngineLoop.cpp
--- a/Engine/Source/Runtime/Launch/Private/LaunchEngineLoop.cpp
+++ b/Engine/Source/Runtime/Launch/Private/LaunchEngineLoop.cpp
@@ -2651,11 +2651,6 @@ bool FEngineLoop::AppInit( )
 
 #if STATS
 	FThreadStats::StartThread();
-	if (FThreadStats::WillEverCollectData())
-	{
-		FThreadStats::ExplicitFlush(); // flush the stats and set update the scope so we don't flush again until a frame update, this helps prevent fragmentation
-	}
-	FStartupMessages::Get().AddThreadMetadata( NAME_GameThread, FPlatformTLS::GetCurrentThreadId() );
 #endif
 
 #if WITH_ENGINE
