Profiler - Fixes

diff --git a/Engine/Source/Developer/Profiler/Private/ProfilerManager.cpp b/Engine/Source/Developer/Profiler/Private/ProfilerManager.cpp
--- a/Engine/Source/Developer/Profiler/Private/ProfilerManager.cpp
+++ b/Engine/Source/Developer/Profiler/Private/ProfilerManager.cpp
@@ -200,7 +200,7 @@ void FProfilerManager::LoadProfilerCapture( const FString& ProfilerCaptureFilepa
 
 	ProfilerSessionInstances.Add( ProfilerInstanceID, ProfilerSession );
 	{
-		PROFILER_SCOPE_LOG_TIME( TEXT( "ProfilerClient->LoadCapture" ), );
+		PROFILER_SCOPE_LOG_TIME( TEXT( "ProfilerClient->LoadCapture" ), nullptr );
 		ProfilerClient->LoadCapture( ProfilerCaptureFilepath, ProfilerInstanceID );
 	}
 
@@ -267,6 +267,8 @@ void FProfilerManager::ProfilerSession_OnCaptureFileProcessed( const FGuid Profi
 	const FProfilerSessionRef* ProfilerSession = FindSessionInstance( ProfilerInstanceID );
 	if( ProfilerSession && ProfilerWindow.IsValid())
 	{
+		RequestFilterAndPresetsUpdateEvent.Broadcast();
+
 		GetProfilerWindow()->UpdateEventGraph( ProfilerInstanceID, (*ProfilerSession)->GetEventGraphDataAverage(), (*ProfilerSession)->GetEventGraphDataMaximum(), true );
 		bHasCaptureFileFullyProcessed = true;
 	}
@@ -309,7 +311,10 @@ void FProfilerManager::ProfilerClient_OnMetaDataUpdated( const FGuid& InstanceID
 	{
 		(*ProfilerSession)->UpdateMetadata( ProfilerClient->GetStatMetaData( InstanceID ) );
 
-		RequestFilterAndPresetsUpdateEvent.Broadcast();
+		if( (*ProfilerSession)->GetSessionType() == EProfilerSessionTypes::Live )
+		{
+			RequestFilterAndPresetsUpdateEvent.Broadcast();
+		}
 	}
 }
 
@@ -336,7 +341,6 @@ void FProfilerManager::ProfilerClient_OnLoadCompleted( const FGuid& InstanceID )
 	const FProfilerSessionRef* ProfilerSession = FindSessionInstance( InstanceID );
 	if( ProfilerSession && GetProfilerWindow().IsValid())
 	{
-		RequestFilterAndPresetsUpdateEvent.Broadcast();
 		(*ProfilerSession)->LoadComplete();
 
 		// Update the notification that a file has been loaded, to be precise it should be loaded on the next tick...
@@ -733,7 +737,7 @@ void FProfilerManager::CloseAllEventGraphTabs()
 
 void FProfilerManager::DataGraph_OnSelectionChangedForIndex( uint32 FrameStartIndex, uint32 FrameEndIndex )
 {
-	PROFILER_SCOPE_LOG_TIME( TEXT( "FProfilerManager::DataGraph_OnSelectionChangedForIndex" ), );
+	PROFILER_SCOPE_LOG_TIME( TEXT( "FProfilerManager::DataGraph_OnSelectionChangedForIndex" ), nullptr );
 
 	for( auto It = GetProfilerInstancesIterator(); It; ++It )
 	{
diff --git a/Engine/Source/Developer/Profiler/Private/ProfilerSession.cpp b/Engine/Source/Developer/Profiler/Private/ProfilerSession.cpp
--- a/Engine/Source/Developer/Profiler/Private/ProfilerSession.cpp
+++ b/Engine/Source/Developer/Profiler/Private/ProfilerSession.cpp
@@ -243,6 +243,9 @@ void FProfilerSession::UpdateAggregatedEventGraphData( const uint32 FrameIndex )
 
 bool FProfilerSession::HandleTicker( float DeltaTime )
 {
+	static FTotalTimeAndCount Current( 0.0f, 0 );
+	PROFILER_SCOPE_LOG_TIME( TEXT( "FProfilerSession::HandleTicker" ), &Current );
+
 	// Update metadata if needed
 	if( bRequestStatMetadataUpdate )
 	{
@@ -250,18 +253,24 @@ bool FProfilerSession::HandleTicker( float DeltaTime )
 		bRequestStatMetadataUpdate = false;
 	}
 
-	const int32 FramesNumToProcess = SessionType == EProfilerSessionTypes::StatsFile ? 240 : 2;
-
-	for( int32 FrameNumber = 0; FrameNumber < FMath::Min(FrameToProcess.Num(),FramesNumToProcess); FrameNumber++ )
+	// Limit processing to 30ms per frame.
+	const double TimeLimit = 30 / 1000.0;
+	double Seconds = 0;
+	
+	for( int32 FrameNumber = 0; FrameNumber < FrameToProcess.Num(); FrameNumber++ )
 	{	
+		if( Seconds > TimeLimit )
+		{
+			break;
+		}
+
+		FScopeSecondsCounter SecondsCounter(Seconds);
+
 		const uint32 FrameIndex = FrameToProcess[0];
 		FrameToProcess.RemoveAt( 0 );
 
 		FProfilerDataFrame& CurrentProfilerData = FrameToProfilerDataMapping.FindChecked( FrameIndex );
 
-		static FTotalTimeAndCount Current(0.0f, 0);
-		PROFILER_SCOPE_LOG_TIME( TEXT( "FProfilerSession::HandleTicker" ), &Current );
-
 		const FProfilerStatMetaDataRef MetaData = GetMetaData();
 		TMap<uint32, float> ThreadMS;
 
@@ -840,7 +849,7 @@ void FRawProfilerSession::PrepareLoading()
 
 void FProfilerStatMetaData::Update( const FStatMetaData& ClientStatMetaData )
 {
-	PROFILER_SCOPE_LOG_TIME( TEXT( "FProfilerStatMetaData.Update" ), );
+	PROFILER_SCOPE_LOG_TIME( TEXT( "FProfilerStatMetaData.Update" ), nullptr );
 
 	FStatMetaData LocalCopy;
 	{
diff --git a/Engine/Source/Developer/Profiler/Private/Widgets/SEventGraph.cpp b/Engine/Source/Developer/Profiler/Private/Widgets/SEventGraph.cpp
--- a/Engine/Source/Developer/Profiler/Private/Widgets/SEventGraph.cpp
+++ b/Engine/Source/Developer/Profiler/Private/Widgets/SEventGraph.cpp
@@ -1493,7 +1493,7 @@ static EEventCompareOps::Type EColumnSortModeToEventCompareOp( const EColumnSort
 
 void SEventGraph::SortEvents()
 {
-	PROFILER_SCOPE_LOG_TIME( TEXT( "SEventGraph::SortEvents" ), );
+	PROFILER_SCOPE_LOG_TIME( TEXT( "SEventGraph::SortEvents" ), nullptr );
 
 	if( ColumnBeingSorted != NAME_None )
 	{
@@ -2919,7 +2919,7 @@ TSharedRef<SHorizontalBox> SEventGraph::GetContentForEvent( FEventGraphSamplePtr
 
 void SEventGraph::SetNewEventGraphState( const FEventGraphDataRef AverageEventGraph, const FEventGraphDataRef MaximumEventGraph, bool bInitial )
 {
-	PROFILER_SCOPE_LOG_TIME( TEXT( "SEventGraph::UpdateEventGraph" ), );
+	PROFILER_SCOPE_LOG_TIME( TEXT( "SEventGraph::UpdateEventGraph" ), nullptr );
 
 	// Store current operation.
 	SaveCurrentEventGraphState();
@@ -3193,7 +3193,7 @@ bool SEventGraph::ClearHistory_CanExecute() const
 
 void SEventGraph::FilterOutByProperty_Execute( const FEventGraphSamplePtr EventPtr, const FName PropertyName, const bool bReset )
 {
-	PROFILER_SCOPE_LOG_TIME( TEXT( "SEventGraph::FilterOutByProperty_Execute" ), );
+	PROFILER_SCOPE_LOG_TIME( TEXT( "SEventGraph::FilterOutByProperty_Execute" ), nullptr );
 
 	// Store current operation.
 	SaveCurrentEventGraphState();
@@ -3217,7 +3217,7 @@ bool SEventGraph::FilterOutByProperty_CanExecute( const FEventGraphSamplePtr Eve
 
 void SEventGraph::CullByProperty_Execute( const FEventGraphSamplePtr EventPtr, const FName PropertyName, const bool bReset )
 {
-	PROFILER_SCOPE_LOG_TIME( TEXT( "SEventGraph::CullByProperty_Execute" ), );
+	PROFILER_SCOPE_LOG_TIME( TEXT( "SEventGraph::CullByProperty_Execute" ), nullptr );
 
 	// Store current operation.
 	SaveCurrentEventGraphState();
diff --git a/Engine/Source/Developer/ProfilerClient/Private/ProfilerClientManager.cpp b/Engine/Source/Developer/ProfilerClient/Private/ProfilerClientManager.cpp
--- a/Engine/Source/Developer/ProfilerClient/Private/ProfilerClientManager.cpp
+++ b/Engine/Source/Developer/ProfilerClient/Private/ProfilerClientManager.cpp
@@ -1114,10 +1114,7 @@ bool FServiceConnection::ReadAndConvertStatMessages( FArchive& Reader, bool bUse
 						{
 							if( DataFrames.Num() > FProfilerClientManager::MaxFramesPerTick )
 							{
-								while( DataFrames.Num() )
-								{
-									FPlatformProcess::Sleep( 0.001f );
-								}
+								FPlatformProcess::Sleep( 0.1f );
 							}
 						}
 
diff --git a/Engine/Source/Developer/ProfilerClient/Private/ProfilerClientManager.h b/Engine/Source/Developer/ProfilerClient/Private/ProfilerClientManager.h
--- a/Engine/Source/Developer/ProfilerClient/Private/ProfilerClientManager.h
+++ b/Engine/Source/Developer/ProfilerClient/Private/ProfilerClientManager.h
@@ -214,7 +214,7 @@ public:
 		return ProfilerLoadedMetaDataDelegate;
 	}
 
-	static const int32 MaxFramesPerTick = 60;
+	static const int32 MaxFramesPerTick = 30;
 
 private:
 
