Slightly optimize string column sorting in profile tool

diff --git a/Engine/Source/Developer/Profiler/Private/ProfilerDataSource.h b/Engine/Source/Developer/Profiler/Private/ProfilerDataSource.h
--- a/Engine/Source/Developer/Profiler/Private/ProfilerDataSource.h
+++ b/Engine/Source/Developer/Profiler/Private/ProfilerDataSource.h
@@ -1926,7 +1926,7 @@ namespace EventGraphPrivate
 
 		FORCEINLINE_DEBUGGABLE bool operator()( const FEventGraphSamplePtr& A, const FEventGraphSamplePtr& B ) const 
 		{
-			return A.Get()->GetPropertyValueAsString(PropertyIndex) > B.Get()->GetPropertyValueAsString(PropertyIndex);
+			return A.Get()->PropertyValueAsFName(PropertyIndex) > B.Get()->PropertyValueAsFName(PropertyIndex);
 		}
 	};
 
@@ -1960,7 +1960,7 @@ namespace EventGraphPrivate
 
 		FORCEINLINE_DEBUGGABLE bool operator()( const FEventGraphSamplePtr& A, const FEventGraphSamplePtr& B ) const 
 		{
-			return A.Get()->GetPropertyValueAsString(PropertyIndex) < B.Get()->GetPropertyValueAsString(PropertyIndex);
+			return A.Get()->PropertyValueAsFName(PropertyIndex) < B.Get()->PropertyValueAsFName(PropertyIndex);
 		}
 	};
 
@@ -2015,13 +2015,10 @@ public:
 		{
 			if( bCompareAsStrings )
 			{
-				// @TODO: For string use Contains instead of Less/Greater
-				const auto Value = SrcPtr->GetPropertyValueAsString(SrcEventProperty.Index);
 				ExecuteOperationInternal( DestPtr->GetChildren(), DestEventProperty, SrcPtr, EventGraphPrivate::TCompareByProperty<EEventPropertyTypes::Name,EEventCompareOps::Less>(SrcEventProperty.Index) );
 			}
 			else if( bCompareAsFloat )
 			{
-				const auto Value = SrcPtr->PropertyValueAsDouble(SrcEventProperty.Index);
 				ExecuteOperationInternal( DestPtr->GetChildren(), DestEventProperty, SrcPtr, EventGraphPrivate::TCompareByProperty<EEventPropertyTypes::Double,EEventCompareOps::Less>(SrcEventProperty.Index) );
 			}
 
@@ -2030,12 +2027,10 @@ public:
 		{
 			if( bCompareAsStrings )
 			{
-				const auto Value = SrcPtr->GetPropertyValueAsString(SrcEventProperty.Index);
 				ExecuteOperationInternal( DestPtr->GetChildren(), DestEventProperty, SrcPtr, EventGraphPrivate::TCompareByProperty<EEventPropertyTypes::Name,EEventCompareOps::Greater>(SrcEventProperty.Index) );
 			}
 			else if( bCompareAsFloat )
 			{
-				const auto Value = SrcPtr->PropertyValueAsDouble(SrcEventProperty.Index);
 				ExecuteOperationInternal( DestPtr->GetChildren(), DestEventProperty, SrcPtr, EventGraphPrivate::TCompareByProperty<EEventPropertyTypes::Double,EEventCompareOps::Greater>(SrcEventProperty.Index) );
 			}
 		}
