Use FName.Compare instead of < and (non-existent) > FName operators

diff --git a/Engine/Source/Developer/Profiler/Private/ProfilerDataSource.h b/Engine/Source/Developer/Profiler/Private/ProfilerDataSource.h
--- a/Engine/Source/Developer/Profiler/Private/ProfilerDataSource.h
+++ b/Engine/Source/Developer/Profiler/Private/ProfilerDataSource.h
@@ -1926,7 +1926,7 @@ namespace EventGraphPrivate
 
 		FORCEINLINE_DEBUGGABLE bool operator()( const FEventGraphSamplePtr& A, const FEventGraphSamplePtr& B ) const 
 		{
-			return A.Get()->PropertyValueAsFName(PropertyIndex) > B.Get()->PropertyValueAsFName(PropertyIndex);
+			return A.Get()->PropertyValueAsFName(PropertyIndex).Compare(B.Get()->PropertyValueAsFName(PropertyIndex)) > 0;
 		}
 	};
 
@@ -1960,7 +1960,7 @@ namespace EventGraphPrivate
 
 		FORCEINLINE_DEBUGGABLE bool operator()( const FEventGraphSamplePtr& A, const FEventGraphSamplePtr& B ) const 
 		{
-			return A.Get()->PropertyValueAsFName(PropertyIndex) < B.Get()->PropertyValueAsFName(PropertyIndex);
+			return A.Get()->PropertyValueAsFName(PropertyIndex).Compare(B.Get()->PropertyValueAsFName(PropertyIndex)) < 0;
 		}
 	};
 
