Improve loading of large profiler data sets by not using AddUnique

diff --git a/Engine/Source/Developer/Profiler/Private/ProfilerSession.h b/Engine/Source/Developer/Profiler/Private/ProfilerSession.h
--- a/Engine/Source/Developer/Profiler/Private/ProfilerSession.h
+++ b/Engine/Source/Developer/Profiler/Private/ProfilerSession.h
@@ -268,7 +268,7 @@ protected:
 	 */
 	void AddStat( FProfilerStat* ProfilerStat )
 	{
-		OwnedStats.AddUnique( ProfilerStat );
+		OwnedStats.Add( ProfilerStat );
 	}
 
 public:
@@ -433,9 +433,6 @@ private:
 	 */
 	void InitializeStat( const uint32 StatID, const uint32 GroupID, const FString& StatName, const EStatType InType, FName StatFName = NAME_None )	
 	{
-		// @TODO yrx 2014-03-24 FindRef!
-		FProfilerGroup* GroupPtr = GroupDescriptions.FindRef( GroupID );
-
 		FProfilerStat* StatPtr = StatDescriptions.FindRef( StatID );
 		if( !StatPtr )
 		{
@@ -445,44 +442,49 @@ private:
 			{
 				StatFNameDescriptions.Add( StatFName, StatPtr );
 			}
-		}
-		StatPtr->Initialize( StatName, GroupPtr, InType );
 
-		if( StatFName == NAME_None && GroupPtr->Name() == TEXT( "Threads" ) )
-		{
-			// Check if this stat is a thread stat.
-			const uint32 ThreadID = FStatsUtils::ParseThreadID( StatPtr->Name().ToString() );
-			if( ThreadID != 0 )
+			// @TODO yrx 2014-03-24 FindRef!
+			FProfilerGroup* GroupPtr = GroupDescriptions.FindRef( GroupID );
+
+			StatPtr->Initialize( StatName, GroupPtr, InType );
+
+			static const FName NAME_Threads(TEXT("Threads"));
+			if( StatFName == NAME_None && GroupPtr->Name() == NAME_Threads )
 			{
-				const FString* ThreadDesc = GetThreadDescriptions().Find( ThreadID );
-				if( ThreadDesc )
+				// Check if this stat is a thread stat.
+				const uint32 ThreadID = FStatsUtils::ParseThreadID( StatPtr->Name().ToString() );
+				if( ThreadID != 0 )
 				{
-					// Replace the stat name with a thread name.
-					const FString UniqueThreadName = FString::Printf( TEXT( "%s [0x%x]" ), **ThreadDesc, ThreadID );
-					StatPtr->_Name = *UniqueThreadName;
-					ThreadIDtoStatID.Add( ThreadID, StatID );
-
-					// Game thread is always NAME_GameThread
-					if( **ThreadDesc == FName( NAME_GameThread ) )
+					const FString* ThreadDesc = GetThreadDescriptions().Find( ThreadID );
+					if( ThreadDesc )
 					{
-						GameThreadID = ThreadID;
-					}
-					// Rendering thread may be "Rendering thread" or NAME_RenderThread with an index
-					else if( ThreadDesc->Contains( FName( NAME_RenderThread ).GetPlainNameString() ) )
-					{
-						RenderThreadIDs.AddUnique( ThreadID );
-					}
-					else if( ThreadDesc->Contains( TEXT( "RenderingThread" ) ) )
-					{
-						RenderThreadIDs.AddUnique( ThreadID );
-					}
-				}			
+						// Replace the stat name with a thread name.
+						const FString UniqueThreadName = FString::Printf( TEXT( "%s [0x%x]" ), **ThreadDesc, ThreadID );
+						StatPtr->_Name = *UniqueThreadName;
+						ThreadIDtoStatID.Add( ThreadID, StatID );
+
+						// Game thread is always NAME_GameThread
+						if( **ThreadDesc == FName( NAME_GameThread ) )
+						{
+							GameThreadID = ThreadID;
+						}
+						// Rendering thread may be "Rendering thread" or NAME_RenderThread with an index
+						else if( ThreadDesc->Contains( FName( NAME_RenderThread ).GetPlainNameString() ) )
+						{
+							RenderThreadIDs.AddUnique( ThreadID );
+						}
+						else if( ThreadDesc->Contains( TEXT( "RenderingThread" ) ) )
+						{
+							RenderThreadIDs.AddUnique( ThreadID );
+						}
+					}			
+				}
 			}
-		}
 
-		if( GroupPtr != FProfilerGroup::GetDefaultPtr() )
-		{
-			GroupPtr->AddStat( StatPtr );
+			if( GroupPtr != FProfilerGroup::GetDefaultPtr() )
+			{
+				GroupPtr->AddStat( StatPtr );
+			}
 		}
 	}
 
