Profiler: Added ability to resize named events list

diff --git a/Engine/Source/Developer/Profiler/Private/Widgets/SEventGraph.cpp b/Engine/Source/Developer/Profiler/Private/Widgets/SEventGraph.cpp
--- a/Engine/Source/Developer/Profiler/Private/Widgets/SEventGraph.cpp
+++ b/Engine/Source/Developer/Profiler/Private/Widgets/SEventGraph.cpp
@@ -761,166 +761,173 @@ void SEventGraph::Construct( const FArguments& InArgs )
 
 	ChildSlot
 	[
-		SNew( SVerticalBox )
+		SNew(SSplitter)
+			.Orientation(Orient_Vertical)
 
-		+SVerticalBox::Slot()
-		.AutoHeight()
-		[
-			SNew( SBorder )
-			.BorderImage( FEditorStyle::GetBrush("ToolPanel.GroupBorder") )
-			.Padding( 2.0f )
+		+ SSplitter::Slot()
+			.Value(0.7f)
 			[
-				SNew( SHorizontalBox )
-
-				+SHorizontalBox::Slot()
-				.AutoWidth()
-				[
-					GetWidgetForEventGraphTypes()
-				]
-
-				+SHorizontalBox::Slot()
-				.AutoWidth()
-				[
-					SNew( SSpacer )
-					.Size( FVector2D( 2.0f, 2.0f ) )
-				]
-
-				+SHorizontalBox::Slot()
-				.AutoWidth()
-				[
-					GetWidgetForEventGraphViewModes()
-				]
-
-				+SHorizontalBox::Slot()
-				.AutoWidth()
-				[
-					SNew( SSpacer )
-					.Size( FVector2D( 2.0f, 2.0f ) )
-				]
-
-				+SHorizontalBox::Slot()
-				.FillWidth( 1.0f )
-				[
-					GetWidgetBoxForOptions()
-				]
-			]	
-		]
- 
-		// Spacer
-		+SVerticalBox::Slot()
-		.AutoHeight()
-		[
-			SNew( SSpacer )
-			.Size( FVector2D( 2.0f, 2.0f ) )
-		]
-
-		// Function details view ( @see VS2012 profiler )
-		+SVerticalBox::Slot()
-		.AutoHeight()
-		[		
-			SAssignNew(FunctionDetailsBox,SBox)
-			.HeightOverride( 224.0f )
-			[
-				SNew( SBorder )
-				.BorderImage( FEditorStyle::GetBrush("ToolPanel.GroupBorder") )
-				.Padding( 2.0f )
-				[
-					SNew( SHorizontalBox )
+				SNew( SVerticalBox )
 
-					// Calling Functions
-					+SHorizontalBox::Slot()
-					.FillWidth( 1.0f )
-					.Padding( 2.0f )
+				+ SVerticalBox::Slot()
+					.AutoHeight()
 					[
-						GetVerticalBoxForFunctionDetails( VerticalBox_TopCalling, LOCTEXT("FunctionDetails_CallingFunctions","Calling Functions") )
+						SNew( SBorder )
+							.BorderImage( FEditorStyle::GetBrush("ToolPanel.GroupBorder") )
+							.Padding( 2.0f )
+							[
+								SNew( SHorizontalBox )
+
+								+ SHorizontalBox::Slot()
+									.AutoWidth()
+									[
+										GetWidgetForEventGraphTypes()
+									]
+
+								+ SHorizontalBox::Slot()
+									.AutoWidth()
+									[
+										SNew( SSpacer )
+											.Size( FVector2D( 2.0f, 2.0f ) )
+									]
+
+								+ SHorizontalBox::Slot()
+									.AutoWidth()
+									[
+										GetWidgetForEventGraphViewModes()
+									]
+
+								+ SHorizontalBox::Slot()
+									.AutoWidth()
+									[
+										SNew( SSpacer )
+											.Size( FVector2D( 2.0f, 2.0f ) )
+									]
+
+								+ SHorizontalBox::Slot()
+									.FillWidth( 1.0f )
+									[
+										GetWidgetBoxForOptions()
+									]
+							]
 					]
-
-					// ->
-					+SHorizontalBox::Slot()
-					.AutoWidth()
-					.HAlign( HAlign_Center )
-					.VAlign( VAlign_Top )
-					.Padding( 2.0f )
-					[
-						SNew( SImage )
-						.Image( FEditorStyle::GetBrush("Animation.Forward") )
-					]
-
-					// Current Function
-					+SHorizontalBox::Slot()
-					.FillWidth( 1.0f )
-					.Padding( 2.0f )
-					[
-						GetVerticalBoxForCurrentFunction()
-					]
-
-					// ->
-					+SHorizontalBox::Slot()
-					.AutoWidth()
-					.HAlign( HAlign_Center )
-					.VAlign( VAlign_Top )
-					.Padding( 2.0f )
+ 
+				// Spacer
+				+ SVerticalBox::Slot()
+					.AutoHeight()
 					[
-						SNew( SImage )
-						.Image( FEditorStyle::GetBrush("Animation.Forward") )
+						SNew( SSpacer )
+							.Size( FVector2D( 2.0f, 2.0f ) )
 					]
 
-					// Called Functions
-					+SHorizontalBox::Slot()
-					.FillWidth( 1.0f )
-					.Padding( 2.0f )
-					[
-						GetVerticalBoxForFunctionDetails( VerticalBox_TopCalled, LOCTEXT("FunctionDetails_CalledFunctions","Called Functions") )
+				// Function details view ( @see VS2012 profiler )
+				+ SVerticalBox::Slot()
+					.AutoHeight()
+					[		
+						SAssignNew(FunctionDetailsBox,SBox)
+							.HeightOverride( 224.0f )
+							[
+								SNew( SBorder )
+									.BorderImage( FEditorStyle::GetBrush("ToolPanel.GroupBorder") )
+									.Padding( 2.0f )
+									[
+										SNew( SHorizontalBox )
+
+										// Calling Functions
+										+ SHorizontalBox::Slot()
+											.FillWidth( 1.0f )
+											.Padding( 2.0f )
+											[
+												GetVerticalBoxForFunctionDetails( VerticalBox_TopCalling, LOCTEXT("FunctionDetails_CallingFunctions","Calling Functions") )
+											]
+
+										// ->
+										+ SHorizontalBox::Slot()
+											.AutoWidth()
+											.HAlign( HAlign_Center )
+											.VAlign( VAlign_Top )
+											.Padding( 2.0f )
+											[
+												SNew( SImage )
+													.Image( FEditorStyle::GetBrush("Animation.Forward") )
+											]
+
+										// Current Function
+										+ SHorizontalBox::Slot()
+											.FillWidth( 1.0f )
+											.Padding( 2.0f )
+											[
+												GetVerticalBoxForCurrentFunction()
+											]
+
+										// ->
+										+ SHorizontalBox::Slot()
+											.AutoWidth()
+											.HAlign( HAlign_Center )
+											.VAlign( VAlign_Top )
+											.Padding( 2.0f )
+											[
+												SNew( SImage )
+													.Image( FEditorStyle::GetBrush("Animation.Forward") )
+											]
+
+										// Called Functions
+										+ SHorizontalBox::Slot()
+											.FillWidth( 1.0f )
+											.Padding( 2.0f )
+											[
+												GetVerticalBoxForFunctionDetails( VerticalBox_TopCalled, LOCTEXT("FunctionDetails_CalledFunctions","Called Functions") )
+											]
+									]
+							]
 					]
-				]
 			]
-		]
 
-		+SVerticalBox::Slot()
-		.FillHeight( 1.0f )
-		.Padding(0.0f, 6.0f, 0.0f, 0.0f)
-		[
-			SNew( SBorder )
-			.BorderImage( FEditorStyle::GetBrush("ToolPanel.GroupBorder") )
-			.Padding(0.0f)
+		+ SSplitter::Slot()
+			.Value(0.3f)
 			[
-				SNew(SHorizontalBox)
-
-				+SHorizontalBox::Slot()
-				.FillWidth( 1.0f )
-				[
-					SNew( SHorizontalBox )
-					+SHorizontalBox::Slot()
-					//.Padding( 2.0f )
-					[
-						SAssignNew( TreeView_Base, STreeView< FEventGraphSamplePtr > )
-						.ExternalScrollbar( ExternalScrollbar )
-						.SelectionMode( ESelectionMode::Multi )
-						.TreeItemsSource( &StaticEventArray )
-						.OnGetChildren( this, &SEventGraph::EventGraph_OnGetChildren )
-						.OnGenerateRow( this, &SEventGraph::EventGraph_OnGenerateRow )
-						.OnSelectionChanged( this, &SEventGraph::EventGraph_OnSelectionChanged )
-						.OnContextMenuOpening( FOnContextMenuOpening::CreateSP( this, &SEventGraph::EventGraph_GetMenuContent ) )
-						.ItemHeight( 12 )
-						.HeaderRow
-						(
-							SAssignNew(TreeViewHeaderRow,SHeaderRow)
-							.Visibility(EVisibility::Visible)
-						)
-					]
-				]
-
-				+SHorizontalBox::Slot()
-				.AutoWidth()
-				[
-					SNew( SBox )
-					.WidthOverride( FOptionalSize( 16 ) )
+				SNew( SBorder )
+					.BorderImage( FEditorStyle::GetBrush("ToolPanel.GroupBorder") )
+					.Padding(0.0f)
 					[
-						ExternalScrollbar.ToSharedRef()
+						SNew(SHorizontalBox)
+
+						+SHorizontalBox::Slot()
+							.FillWidth( 1.0f )
+							[
+								SNew( SHorizontalBox )
+
+								+ SHorizontalBox::Slot()
+									//.Padding( 2.0f )
+									[
+										SAssignNew( TreeView_Base, STreeView< FEventGraphSamplePtr > )
+											.ExternalScrollbar( ExternalScrollbar )
+											.SelectionMode( ESelectionMode::Multi )
+											.TreeItemsSource( &StaticEventArray )
+											.OnGetChildren( this, &SEventGraph::EventGraph_OnGetChildren )
+											.OnGenerateRow( this, &SEventGraph::EventGraph_OnGenerateRow )
+											.OnSelectionChanged( this, &SEventGraph::EventGraph_OnSelectionChanged )
+											.OnContextMenuOpening( FOnContextMenuOpening::CreateSP( this, &SEventGraph::EventGraph_GetMenuContent ) )
+											.ItemHeight( 12 )
+											.HeaderRow
+											(
+												SAssignNew(TreeViewHeaderRow,SHeaderRow)
+													.Visibility(EVisibility::Visible)
+											)
+									]
+							]
+
+						+ SHorizontalBox::Slot()
+							.AutoWidth()
+							[
+								SNew( SBox )
+									.WidthOverride( FOptionalSize( 16 ) )
+									[
+										ExternalScrollbar.ToSharedRef()
+									]
+							]
 					]
-				]
 			]
-		]
 	];
 
 	InitializeAndShowHeaderColumns();
