Converted FUIAction to use ECheckBoxState when reporting its check state

There are now overloads to FUIAction that take a FGetActionCheckState delegate which is able to return any of the entries from ECheckBoxState. FIsActionChecked still exists for backwards compatibility, and converts the returned bool into a ECheckBoxState via a passthrough function.

diff --git a/Engine/Config/BaseEngine.ini b/Engine/Config/BaseEngine.ini
--- a/Engine/Config/BaseEngine.ini
+++ b/Engine/Config/BaseEngine.ini
@@ -1176,7 +1176,8 @@ bUseStreamingPause=false
 
 ; UMG Checkbox State Rename
 
-+ActiveClassRedirects=(OldClassName="ESlateCheckBoxState", NewClassName="ECheckBoxState")
++ActiveClassRedirects=(OldClassName="ESlateCheckBoxState", NewClassName="/Script/SlateCore.ECheckBoxState")
++ActiveClassRedirects=(OldClassName="ECheckBoxState",NewClassName="/Script/SlateCore.ECheckBoxState")
 
 ; SplineComponent API overhaul
 +K2FieldRedirects=(OldFieldName="SplineComponent.GetNumSplinePoints",NewFieldName="SplineComponent.GetNumberOfSplinePoints")
diff --git a/Engine/Source/Developer/Profiler/Private/ProfilerCommands.cpp b/Engine/Source/Developer/Profiler/Private/ProfilerCommands.cpp
--- a/Engine/Source/Developer/Profiler/Private/ProfilerCommands.cpp
+++ b/Engine/Source/Developer/Profiler/Private/ProfilerCommands.cpp
@@ -109,8 +109,7 @@ const FUIAction FProfilerActionManager::ToggleDataPreview_Custom( const FGuid Se
 	FUIAction UIAction;
 	UIAction.ExecuteAction = FExecuteAction::CreateRaw( this, &FProfilerActionManager::ToggleDataPreview_Execute, SessionInstanceID );
 	UIAction.CanExecuteAction = FCanExecuteAction::CreateRaw( this, &FProfilerActionManager::ToggleDataPreview_CanExecute, SessionInstanceID );
-	UIAction.IsCheckedDelegate = FIsActionChecked::CreateRaw( this, &FProfilerActionManager::ToggleDataPreview_IsChecked, SessionInstanceID );
-	UIAction.IsActionVisibleDelegate = FIsActionButtonVisible();
+	UIAction.GetActionCheckState = FGetActionCheckState::CreateRaw( this, &FProfilerActionManager::ToggleDataPreview_GetCheckState, SessionInstanceID );
 	return UIAction;
 }
 
@@ -162,7 +161,7 @@ bool FProfilerActionManager::ToggleDataPreview_CanExecute( const FGuid SessionIn
 	}
 }
 
-bool FProfilerActionManager::ToggleDataPreview_IsChecked( const FGuid SessionInstanceID ) const
+ECheckBoxState FProfilerActionManager::ToggleDataPreview_GetCheckState( const FGuid SessionInstanceID ) const
 {
 	// One session instance
 	if( SessionInstanceID.IsValid() )
@@ -170,18 +169,18 @@ bool FProfilerActionManager::ToggleDataPreview_IsChecked( const FGuid SessionIns
 		const FProfilerSessionRef* ProfilerSessionPtr = This->FindSessionInstance( SessionInstanceID );
 		if( ProfilerSessionPtr )
 		{
-			return (*ProfilerSessionPtr)->bDataPreviewing;
+			return (*ProfilerSessionPtr)->bDataPreviewing ? ECheckBoxState::Checked : ECheckBoxState::Unchecked;
 		}
 		else
 		{
-			return false;
+			return ECheckBoxState::Unchecked;
 		}
 	}
 	// All session instances
 	else
 	{
 		const bool bDataPreview = This->IsDataPreviewing();
-		return bDataPreview;
+		return bDataPreview ? ECheckBoxState::Checked : ECheckBoxState::Unchecked;
 	}
 }
 
@@ -194,8 +193,7 @@ void FProfilerActionManager::Map_ProfilerManager_ToggleLivePreview_Global()
 	FUIAction UIAction;
 	UIAction.ExecuteAction = FExecuteAction::CreateRaw( this, &FProfilerActionManager::ProfilerManager_ToggleLivePreview_Execute );
 	UIAction.CanExecuteAction = FCanExecuteAction::CreateRaw( this, &FProfilerActionManager::ProfilerManager_ToggleLivePreview_CanExecute );
-	UIAction.IsCheckedDelegate = FIsActionChecked::CreateRaw( this, &FProfilerActionManager::ProfilerManager_ToggleLivePreview_IsChecked );
-	UIAction.IsActionVisibleDelegate = FIsActionButtonVisible();
+	UIAction.GetActionCheckState = FGetActionCheckState::CreateRaw( this, &FProfilerActionManager::ProfilerManager_ToggleLivePreview_GetCheckState );
 
 	This->CommandList->MapAction( This->GetCommands().ProfilerManager_ToggleLivePreview, UIAction );
 }
@@ -211,9 +209,9 @@ bool FProfilerActionManager::ProfilerManager_ToggleLivePreview_CanExecute() cons
 	return bCanExecute;
 }
 
-bool FProfilerActionManager::ProfilerManager_ToggleLivePreview_IsChecked() const
+ECheckBoxState FProfilerActionManager::ProfilerManager_ToggleLivePreview_GetCheckState() const
 {
-	return This->bLivePreview;
+	return This->bLivePreview ? ECheckBoxState::Checked : ECheckBoxState::Unchecked;
 }
 
 /*-----------------------------------------------------------------------------
@@ -225,7 +223,7 @@ const FUIAction FProfilerActionManager::ToggleShowDataGraph_Custom( const FGuid
 	FUIAction UIAction;
 	UIAction.ExecuteAction = FExecuteAction::CreateRaw( this, &FProfilerActionManager::ToggleShowDataGraph_Execute, SessionInstanceID );
 	UIAction.CanExecuteAction = FCanExecuteAction::CreateRaw( this, &FProfilerActionManager::ToggleShowDataGraph_CanExecute, SessionInstanceID );
-	UIAction.IsCheckedDelegate = FIsActionChecked::CreateRaw( this, &FProfilerActionManager::ToggleShowDataGraph_IsChecked, SessionInstanceID );
+	UIAction.GetActionCheckState = FGetActionCheckState::CreateRaw( this, &FProfilerActionManager::ToggleShowDataGraph_GetCheckState, SessionInstanceID );
 	UIAction.IsActionVisibleDelegate = FIsActionButtonVisible::CreateRaw( this, &FProfilerActionManager::ToggleShowDataGraph_IsActionButtonVisible, SessionInstanceID );
 	return UIAction;
 }
@@ -263,18 +261,18 @@ bool FProfilerActionManager::ToggleShowDataGraph_CanExecute( const FGuid Session
 	}
 }
 
-bool FProfilerActionManager::ToggleShowDataGraph_IsChecked( const FGuid SessionInstanceID ) const
+ECheckBoxState FProfilerActionManager::ToggleShowDataGraph_GetCheckState( const FGuid SessionInstanceID ) const
 {
 	// One session instance
 	if( SessionInstanceID.IsValid() )
 	{
 		const bool bIsSessionInstanceTracked = This->IsSessionInstanceTracked( SessionInstanceID );
-		return bIsSessionInstanceTracked;
+		return bIsSessionInstanceTracked ? ECheckBoxState::Checked : ECheckBoxState::Unchecked;
 	}
 	// All session instances
 	else
 	{
-		return false;
+		return ECheckBoxState::Unchecked;
 	}
 }
 
@@ -302,8 +300,6 @@ void FProfilerActionManager::Map_ProfilerManager_Load()
 	FUIAction UIAction;
 	UIAction.ExecuteAction = FExecuteAction::CreateRaw( this, &FProfilerActionManager::ProfilerManager_Load_Execute );
 	UIAction.CanExecuteAction = FCanExecuteAction::CreateRaw( this, &FProfilerActionManager::ProfilerManager_Load_CanExecute );
-	UIAction.IsCheckedDelegate = FIsActionChecked();
-	UIAction.IsActionVisibleDelegate = FIsActionButtonVisible();
 
 	This->CommandList->MapAction( This->GetCommands().ProfilerManager_Load, UIAction );
 }
@@ -369,8 +365,7 @@ const FUIAction FProfilerActionManager::ToggleDataCapture_Custom( const FGuid Se
 	FUIAction UIAction;
 	UIAction.ExecuteAction = FExecuteAction::CreateRaw( this, &FProfilerActionManager::ToggleDataCapture_Execute, SessionInstanceID );
 	UIAction.CanExecuteAction = FCanExecuteAction::CreateRaw( this, &FProfilerActionManager::ToggleDataCapture_CanExecute, SessionInstanceID );
-	UIAction.IsCheckedDelegate = FIsActionChecked::CreateRaw( this, &FProfilerActionManager::ToggleDataCapture_IsChecked, SessionInstanceID );
-	UIAction.IsActionVisibleDelegate = FIsActionButtonVisible();
+	UIAction.GetActionCheckState = FGetActionCheckState::CreateRaw( this, &FProfilerActionManager::ToggleDataCapture_GetCheckState, SessionInstanceID );
 	return UIAction;
 }
 
@@ -429,7 +424,7 @@ bool FProfilerActionManager::ToggleDataCapture_CanExecute( const FGuid SessionIn
 	return false;
 }
 
-bool FProfilerActionManager::ToggleDataCapture_IsChecked( const FGuid SessionInstanceID ) const
+ECheckBoxState FProfilerActionManager::ToggleDataCapture_GetCheckState( const FGuid SessionInstanceID ) const
 {
 	// One session instance
 	if( SessionInstanceID.IsValid() )
@@ -437,21 +432,21 @@ bool FProfilerActionManager::ToggleDataCapture_IsChecked( const FGuid SessionIns
 		const FProfilerSessionRef* ProfilerSessionPtr = This->FindSessionInstance( SessionInstanceID );
 		if( ProfilerSessionPtr )
 		{
-			return (*ProfilerSessionPtr)->bDataCapturing;
+			return (*ProfilerSessionPtr)->bDataCapturing ? ECheckBoxState::Checked : ECheckBoxState::Unchecked;
 		}
 		else
 		{
-			return false;
+			return ECheckBoxState::Unchecked;
 		}
 	}
 	// All session instances
 	else
 	{
 		const bool bDataCapturing = This->IsDataCapturing();
-		return bDataCapturing;
+		return bDataCapturing ? ECheckBoxState::Checked : ECheckBoxState::Unchecked;
 	}
 
-	return false;
+	return ECheckBoxState::Unchecked;
 }
 
 /*-----------------------------------------------------------------------------
@@ -468,8 +463,6 @@ const FUIAction FProfilerActionManager::OpenSettings_Custom() const
 	FUIAction UIAction;
 	UIAction.ExecuteAction = FExecuteAction::CreateRaw( this, &FProfilerActionManager::OpenSettings_Execute );
 	UIAction.CanExecuteAction = FCanExecuteAction::CreateRaw( this, &FProfilerActionManager::OpenSettings_CanExecute );
-	UIAction.IsCheckedDelegate = FIsActionChecked();
-	UIAction.IsActionVisibleDelegate = FIsActionButtonVisible();
 	return UIAction;
 }
 
diff --git a/Engine/Source/Developer/Profiler/Private/ProfilerCommands.h b/Engine/Source/Developer/Profiler/Private/ProfilerCommands.h
--- a/Engine/Source/Developer/Profiler/Private/ProfilerCommands.h
+++ b/Engine/Source/Developer/Profiler/Private/ProfilerCommands.h
@@ -151,8 +151,8 @@ protected:
 	void ToggleDataPreview_Execute( const FGuid SessionInstanceID );
 	/** Handles FCanExecuteAction for ToggleDataPreview. */
 	bool ToggleDataPreview_CanExecute( const FGuid SessionInstanceID ) const;
-	/** Handles FIsActionChecked for ToggleDataPreview. */
-	bool ToggleDataPreview_IsChecked( const FGuid SessionInstanceID ) const;
+	/** Handles FGetActionCheckState for ToggleDataPreview. */
+	ECheckBoxState ToggleDataPreview_GetCheckState( const FGuid SessionInstanceID ) const;
 
 	/*-----------------------------------------------------------------------------
 		ProfilerManager_ToggleLivePreview
@@ -166,8 +166,8 @@ protected:
 	void ProfilerManager_ToggleLivePreview_Execute();
 	/** Handles FCanExecuteAction for ProfilerManager_ToggleLivePreview. */
 	bool ProfilerManager_ToggleLivePreview_CanExecute( ) const;
-	/** Handles FIsActionChecked for ProfilerManager_ToggleLivePreview. */
-	bool ProfilerManager_ToggleLivePreview_IsChecked() const;
+	/** Handles FGetActionCheckState for ProfilerManager_ToggleLivePreview. */
+	ECheckBoxState ProfilerManager_ToggleLivePreview_GetCheckState() const;
 
 	/*-----------------------------------------------------------------------------
 		ToggleDataCapture
@@ -190,8 +190,8 @@ protected:
 	void ToggleDataCapture_Execute( const FGuid SessionInstanceID );
 	/** Handles FCanExecuteAction for ToggleDataCapture. */
 	bool ToggleDataCapture_CanExecute( const FGuid SessionInstanceID ) const;
-	/** Handles FIsActionChecked for ToggleDataCapture. */
-	bool ToggleDataCapture_IsChecked( const FGuid SessionInstanceID ) const;
+	/** Handles FGetActionCheckState for ToggleDataCapture. */
+	ECheckBoxState ToggleDataCapture_GetCheckState( const FGuid SessionInstanceID ) const;
 
 	/*-----------------------------------------------------------------------------
 		ToggleShowDataGraph
@@ -211,8 +211,8 @@ protected:
 	void ToggleShowDataGraph_Execute( const FGuid SessionInstanceID );
 	/** Handles FCanExecuteAction for ToggleShowDataGraph_Execute. */
 	bool ToggleShowDataGraph_CanExecute( const FGuid SessionInstanceID ) const;
-	/** Handles FIsActionChecked for ToggleShowDataGraph_Execute. */
-	bool ToggleShowDataGraph_IsChecked( const FGuid SessionInstanceID ) const;
+	/** Handles FGetActionCheckState for ToggleShowDataGraph_Execute. */
+	ECheckBoxState ToggleShowDataGraph_GetCheckState( const FGuid SessionInstanceID ) const;
 	/** Handles IsActionButtonVisible for ToggleShowDataGraph_Execute. */
 	bool ToggleShowDataGraph_IsActionButtonVisible( const FGuid SessionInstanceID ) const;
 
diff --git a/Engine/Source/Developer/Profiler/Private/Widgets/SEventGraph.cpp b/Engine/Source/Developer/Profiler/Private/Widgets/SEventGraph.cpp
--- a/Engine/Source/Developer/Profiler/Private/Widgets/SEventGraph.cpp
+++ b/Engine/Source/Developer/Profiler/Private/Widgets/SEventGraph.cpp
@@ -1479,9 +1479,9 @@ bool SEventGraph::ShowSelectedEventsInViewMode_CanExecute( EEventGraphViewModes:
 	return TreeView_Base->GetNumItemsSelected() > 0;
 }
 
-bool SEventGraph::ShowSelectedEventsInViewMode_IsChecked( EEventGraphViewModes::Type NewViewMode ) const
+ECheckBoxState SEventGraph::ShowSelectedEventsInViewMode_GetCheckState( EEventGraphViewModes::Type NewViewMode ) const
 {
-	return GetCurrentStateViewMode() == NewViewMode;
+	return GetCurrentStateViewMode() == NewViewMode ? ECheckBoxState::Checked : ECheckBoxState::Unchecked;
 }
 
 static EEventCompareOps::Type EColumnSortModeToEventCompareOp( const EColumnSortMode::Type ColumnSortMode )
diff --git a/Engine/Source/Developer/Profiler/Private/Widgets/SEventGraph.h b/Engine/Source/Developer/Profiler/Private/Widgets/SEventGraph.h
--- a/Engine/Source/Developer/Profiler/Private/Widgets/SEventGraph.h
+++ b/Engine/Source/Developer/Profiler/Private/Widgets/SEventGraph.h
@@ -297,8 +297,6 @@ public:
 		FUIAction UIAction;
 		UIAction.ExecuteAction = FExecuteAction::CreateSP( this, &SEventGraph::SelectAllFrames_Execute );
 		UIAction.CanExecuteAction = FCanExecuteAction::CreateSP( this, &SEventGraph::SelectAllFrames_CanExecute );
-		UIAction.IsCheckedDelegate = FIsActionChecked();
-		UIAction.IsActionVisibleDelegate = FIsActionButtonVisible();
 		return UIAction;
 	}
 		
@@ -365,8 +363,6 @@ public:
 		FUIAction UIAction;
 		UIAction.ExecuteAction = FExecuteAction::CreateSP( this, &SEventGraph::SetRoot_Execute );
 		UIAction.CanExecuteAction = FCanExecuteAction::CreateSP( this, &SEventGraph::SetRoot_CanExecute );
-		UIAction.IsCheckedDelegate = FIsActionChecked();
-		UIAction.IsActionVisibleDelegate = FIsActionButtonVisible();
 		return UIAction;
 	}
 		
@@ -388,8 +384,6 @@ public:
 		FUIAction UIAction;
 		UIAction.ExecuteAction = FExecuteAction::CreateSP( this, &SEventGraph::ClearHistory_Execute );
 		UIAction.CanExecuteAction = FCanExecuteAction::CreateSP( this, &SEventGraph::ClearHistory_CanExecute );
-		UIAction.IsCheckedDelegate = FIsActionChecked();
-		UIAction.IsActionVisibleDelegate = FIsActionButtonVisible();
 		return UIAction;
 	}
 		
@@ -410,8 +404,7 @@ public:
 		FUIAction UIAction;
 		UIAction.ExecuteAction = FExecuteAction::CreateSP( this, &SEventGraph::ShowSelectedEventsInViewMode_Execute, NewViewMode );
 		UIAction.CanExecuteAction = FCanExecuteAction::CreateSP( this, &SEventGraph::ShowSelectedEventsInViewMode_CanExecute, NewViewMode );
-		UIAction.IsCheckedDelegate = FIsActionChecked::CreateSP( this, &SEventGraph::ShowSelectedEventsInViewMode_IsChecked, NewViewMode );
-		UIAction.IsActionVisibleDelegate = FIsActionButtonVisible();
+		UIAction.GetActionCheckState = FGetActionCheckState::CreateSP( this, &SEventGraph::ShowSelectedEventsInViewMode_GetCheckState, NewViewMode );
 		return UIAction;
 	}
 		
@@ -420,8 +413,8 @@ protected:
 	void ShowSelectedEventsInViewMode_Execute(EEventGraphViewModes::Type NewViewMode);
 	/** Handles FCanExecuteAction for ShowSelectedEventsInViewMode. */
 	bool ShowSelectedEventsInViewMode_CanExecute(EEventGraphViewModes::Type NewViewMode) const;
-	/** Handles FIsActionChecked for ShowSelectedEventsInViewMode. */
-	bool ShowSelectedEventsInViewMode_IsChecked(EEventGraphViewModes::Type NewViewMode) const;
+	/** Handles FGetActionCheckState for ShowSelectedEventsInViewMode. */
+	ECheckBoxState ShowSelectedEventsInViewMode_GetCheckState(EEventGraphViewModes::Type NewViewMode) const;
 
 	/*-----------------------------------------------------------------------------
 		FilterOutByProperty
@@ -434,7 +427,6 @@ public:
 		FUIAction UIAction;
 		UIAction.ExecuteAction = FExecuteAction::CreateSP( this, &SEventGraph::FilterOutByProperty_Execute, EventPtr, PropertyName, bReset );
 		UIAction.CanExecuteAction = FCanExecuteAction::CreateSP( this, &SEventGraph::FilterOutByProperty_CanExecute, EventPtr, PropertyName, bReset );
-		UIAction.IsActionVisibleDelegate = FIsActionButtonVisible();
 		return UIAction;
 	}
 		
@@ -456,7 +448,6 @@ public:
 		FUIAction UIAction;
 		UIAction.ExecuteAction = FExecuteAction::CreateSP( this, &SEventGraph::CullByProperty_Execute, EventPtr, PropertyName, bReset );
 		UIAction.CanExecuteAction = FCanExecuteAction::CreateSP( this, &SEventGraph::CullByProperty_CanExecute, EventPtr, PropertyName, bReset );
-		UIAction.IsActionVisibleDelegate = FIsActionButtonVisible();
 		return UIAction;
 	}
 		
@@ -477,8 +468,7 @@ public:
 		FUIAction UIAction;
 		UIAction.ExecuteAction = FExecuteAction::CreateSP( this, &SEventGraph::HistoryList_GoTo_Execute, StateIndex );
 		UIAction.CanExecuteAction = FCanExecuteAction();
-		UIAction.IsCheckedDelegate = FIsActionChecked::CreateSP( this, &SEventGraph::HistoryList_GoTo_IsChecked, StateIndex );
-		UIAction.IsActionVisibleDelegate = FIsActionButtonVisible();
+		UIAction.GetActionCheckState = FGetActionCheckState::CreateSP( this, &SEventGraph::HistoryList_GoTo_GetCheckState, StateIndex );
 		return UIAction;
 	}
 		
@@ -494,15 +484,10 @@ protected:
 	/** Handles FExecuteAction for HistoryList_GoTo. */
 	void HistoryList_GoTo_Execute( int32 StateIndex );
 
-	ECheckBoxState HistoryList_GoTo_IsCheckedRadioState( int32 StateIndex ) const
+	/** Handles FGetActionCheckState for HistoryList_GoTo. */
+	ECheckBoxState HistoryList_GoTo_GetCheckState( int32 StateIndex ) const
 	{
-		return HistoryList_GoTo_IsChecked( StateIndex ) ? ECheckBoxState::Checked : ECheckBoxState::Unchecked;
-	}
-
-	/** Handles FIsActionChecked for HistoryList_GoTo. */
-	bool HistoryList_GoTo_IsChecked( int32 StateIndex ) const
-	{
-		return StateIndex == CurrentStateIndex;
+		return StateIndex == CurrentStateIndex ? ECheckBoxState::Checked : ECheckBoxState::Unchecked;
 	}
 
 	/*-----------------------------------------------------------------------------
@@ -516,8 +501,6 @@ public:
 		FUIAction UIAction;
 		UIAction.ExecuteAction = FExecuteAction::CreateSP( this, &SEventGraph::SetExpansionForEvents_Execute, SelectedEventType, bShouldExpand );
 		UIAction.CanExecuteAction = FCanExecuteAction::CreateSP( this, &SEventGraph::SetExpansionForEvents_CanExecute, SelectedEventType, bShouldExpand );
-		UIAction.IsCheckedDelegate = FIsActionChecked();
-		UIAction.IsActionVisibleDelegate = FIsActionButtonVisible();
 		return UIAction;
 	}
 		
diff --git a/Engine/Source/Editor/EditorStyle/Private/SlateEditorStyle.cpp b/Engine/Source/Editor/EditorStyle/Private/SlateEditorStyle.cpp
--- a/Engine/Source/Editor/EditorStyle/Private/SlateEditorStyle.cpp
+++ b/Engine/Source/Editor/EditorStyle/Private/SlateEditorStyle.cpp
@@ -1919,12 +1919,15 @@ void FSlateEditorStyle::FStyle::SetupGeneralStyles()
 
 		/* Create style for "ToolBar.CheckBox" widget ... */
 		const FCheckBoxStyle ToolBarCheckBoxStyle = FCheckBoxStyle()
-				.SetUncheckedImage( IMAGE_BRUSH( "Common/SmallCheckBox", Icon14x14 ) )
-				.SetCheckedImage( IMAGE_BRUSH( "Common/SmallCheckBox_Checked",  Icon14x14 ) )
-				.SetUncheckedHoveredImage( IMAGE_BRUSH( "Common/SmallCheckBox_Hovered", Icon14x14 ) )
-				.SetCheckedImage( IMAGE_BRUSH( "Common/SmallCheckBox_Checked_Hovered", Icon14x14 ) )
-				.SetUncheckedPressedImage(IMAGE_BRUSH( "Common/SmallCheckBox_Hovered", Icon14x14, FLinearColor( 0.5f, 0.5f, 0.5f ) ) )
-				.SetCheckedPressedImage( IMAGE_BRUSH( "Common/SmallCheckBox_Checked_Hovered", Icon14x14, FLinearColor( 0.5f, 0.5f, 0.5f ) ) );
+			.SetUncheckedImage( IMAGE_BRUSH( "Common/SmallCheckBox", Icon14x14 ) )
+			.SetCheckedImage( IMAGE_BRUSH( "Common/SmallCheckBox_Checked",  Icon14x14 ) )
+			.SetUncheckedHoveredImage( IMAGE_BRUSH( "Common/SmallCheckBox_Hovered", Icon14x14 ) )
+			.SetCheckedImage( IMAGE_BRUSH( "Common/SmallCheckBox_Checked_Hovered", Icon14x14 ) )
+			.SetUncheckedPressedImage(IMAGE_BRUSH( "Common/SmallCheckBox_Hovered", Icon14x14, FLinearColor( 0.5f, 0.5f, 0.5f ) ) )
+			.SetCheckedPressedImage( IMAGE_BRUSH( "Common/SmallCheckBox_Checked_Hovered", Icon14x14, FLinearColor( 0.5f, 0.5f, 0.5f ) ) )
+			.SetUndeterminedImage( IMAGE_BRUSH( "Common/CheckBox_Undetermined", Icon14x14 ) )
+			.SetUndeterminedHoveredImage( IMAGE_BRUSH( "Common/CheckBox_Undetermined_Hovered", Icon14x14 ) )
+			.SetUndeterminedPressedImage( IMAGE_BRUSH( "Common/CheckBox_Undetermined_Hovered", Icon14x14, FLinearColor( 0.5f, 0.5f, 0.5f ) ) );
 		/* ... and add new style */
 		Set( "ToolBar.CheckBox", ToolBarCheckBoxStyle );
 
@@ -2044,9 +2047,9 @@ void FSlateEditorStyle::FStyle::SetupGeneralStyles()
 			.SetCheckedImage( IMAGE_BRUSH( "Common/SmallCheckBox_Checked", Icon14x14 ) )
 			.SetCheckedHoveredImage( IMAGE_BRUSH( "Common/SmallCheckBox_Checked_Hovered", Icon14x14 ) )
 			.SetCheckedPressedImage( IMAGE_BRUSH( "Common/SmallCheckBox_Checked_Hovered", Icon14x14, FLinearColor( 0.5f, 0.5f, 0.5f ) ) )
-			.SetUndeterminedImage( IMAGE_BRUSH( "Common/CheckBox_Undetermined", Icon16x16 ) )
-			.SetUndeterminedHoveredImage( IMAGE_BRUSH( "Common/CheckBox_Checked_Hovered", Icon16x16 ) )
-			.SetUndeterminedPressedImage( IMAGE_BRUSH( "Common/CheckBox_Undetermined_Hovered", Icon16x16, FLinearColor( 0.5f, 0.5f, 0.5f ) ) );
+			.SetUndeterminedImage( IMAGE_BRUSH( "Common/CheckBox_Undetermined", Icon14x14 ) )
+			.SetUndeterminedHoveredImage( IMAGE_BRUSH( "Common/CheckBox_Undetermined_Hovered", Icon14x14 ) )
+			.SetUndeterminedPressedImage( IMAGE_BRUSH( "Common/CheckBox_Undetermined_Hovered", Icon14x14, FLinearColor( 0.5f, 0.5f, 0.5f ) ) );
  
 		/* ...and add the new style */
 		Set( "Menu.CheckBox", BasicMenuCheckBoxStyle );
@@ -2165,9 +2168,9 @@ void FSlateEditorStyle::FStyle::SetupGeneralStyles()
 			.SetCheckedHoveredImage( IMAGE_BRUSH( "Common/SmallCheckBox_Checked_Hovered", Icon14x14 ) )
 			.SetCheckedPressedImage( IMAGE_BRUSH( "Common/SmallCheckBox_Checked_Hovered", Icon14x14, FLinearColor( 0.5f, 0.5f, 0.5f ) ) )
 			.SetCheckedImage( IMAGE_BRUSH( "Common/SmallCheckBox_Checked",  Icon14x14 ) )
-			.SetUndeterminedImage( IMAGE_BRUSH( "Common/CheckBox_Undetermined", Icon16x16 ) )
-			.SetUndeterminedPressedImage( IMAGE_BRUSH( "Common/CheckBox_Undetermined_Hovered", Icon16x16, FLinearColor( 0.5f, 0.5f, 0.5f ) ) )
-			.SetUndeterminedHoveredImage( IMAGE_BRUSH( "Common/CheckBox_Checked_Hovered", Icon16x16 ) );
+			.SetUndeterminedImage( IMAGE_BRUSH( "Common/CheckBox_Undetermined", Icon14x14 ) )
+			.SetUndeterminedHoveredImage( IMAGE_BRUSH( "Common/CheckBox_Undetermined_Hovered", Icon14x14 ) )
+			.SetUndeterminedPressedImage( IMAGE_BRUSH( "Common/CheckBox_Undetermined_Hovered", Icon14x14, FLinearColor( 0.5f, 0.5f, 0.5f ) ) );
 		Set( "NotificationBar.CheckBox", NotificationBarCheckBoxCheckBoxStyle );
 
 		// Read-only checkbox that appears next to a menu item
diff --git a/Engine/Source/Editor/LevelEditor/Private/SLevelViewportToolBar.cpp b/Engine/Source/Editor/LevelEditor/Private/SLevelViewportToolBar.cpp
--- a/Engine/Source/Editor/LevelEditor/Private/SLevelViewportToolBar.cpp
+++ b/Engine/Source/Editor/LevelEditor/Private/SLevelViewportToolBar.cpp
@@ -681,9 +681,11 @@ void SLevelViewportToolBar::GeneratePlacedCameraMenuEntries( FMenuBuilder& Build
 		AActor* GenericActor = CameraActor;
 
 		FText ActorDisplayName = FText::FromString(CameraActor->GetActorLabel());
-		FUIAction LookThroughCameraAction;
-		LookThroughCameraAction.ExecuteAction.BindSP(Viewport.Pin().ToSharedRef(), &SLevelViewport::OnActorLockToggleFromMenu, GenericActor);
-		LookThroughCameraAction.IsCheckedDelegate.BindSP(Viewport.Pin().ToSharedRef(), &SLevelViewport::IsActorLocked, TWeakObjectPtr<AActor>(GenericActor) );
+		FUIAction LookThroughCameraAction(
+			FExecuteAction::CreateSP(Viewport.Pin().ToSharedRef(), &SLevelViewport::OnActorLockToggleFromMenu, GenericActor),
+			FCanExecuteAction(),
+			FIsActionChecked::CreateSP(Viewport.Pin().ToSharedRef(), &SLevelViewport::IsActorLocked, TWeakObjectPtr<AActor>(GenericActor))
+			);
 
 		Builder.AddMenuEntry( ActorDisplayName, FText::Format(LOCTEXT("LookThroughCameraActor_ToolTip", "Look through and pilot {0}"), ActorDisplayName), CameraIcon, LookThroughCameraAction, NAME_None, EUserInterfaceActionType::RadioButton );
 	}
diff --git a/Engine/Source/Runtime/Slate/Private/Framework/Commands/UICommandList.cpp b/Engine/Source/Runtime/Slate/Private/Framework/Commands/UICommandList.cpp
--- a/Engine/Source/Runtime/Slate/Private/Framework/Commands/UICommandList.cpp
+++ b/Engine/Source/Runtime/Slate/Private/Framework/Commands/UICommandList.cpp
@@ -8,13 +8,13 @@
 
 void FUICommandList::MapAction( const TSharedPtr< const FUICommandInfo > InUICommandInfo, FExecuteAction ExecuteAction, EUIActionRepeatMode RepeatMode )
 {
-	MapAction( InUICommandInfo, ExecuteAction, FCanExecuteAction(), FIsActionChecked(), FIsActionButtonVisible(), RepeatMode );
+	MapAction( InUICommandInfo, ExecuteAction, FCanExecuteAction(), FGetActionCheckState(), FIsActionButtonVisible(), RepeatMode );
 }
 
 
 void FUICommandList::MapAction( const TSharedPtr< const FUICommandInfo > InUICommandInfo, FExecuteAction ExecuteAction, FCanExecuteAction CanExecuteAction, EUIActionRepeatMode RepeatMode )
 {
-	MapAction( InUICommandInfo, ExecuteAction, CanExecuteAction, FIsActionChecked(), FIsActionButtonVisible(), RepeatMode );
+	MapAction( InUICommandInfo, ExecuteAction, CanExecuteAction, FGetActionCheckState(), FIsActionButtonVisible(), RepeatMode );
 }
 
 
@@ -24,12 +24,24 @@ void FUICommandList::MapAction( const TSharedPtr< const FUICommandInfo > InUICom
 }
 
 
+void FUICommandList::MapAction( const TSharedPtr< const FUICommandInfo > InUICommandInfo, FExecuteAction ExecuteAction, FCanExecuteAction CanExecuteAction, FGetActionCheckState GetActionCheckState, EUIActionRepeatMode RepeatMode )
+{
+	MapAction( InUICommandInfo, ExecuteAction, CanExecuteAction, GetActionCheckState, FIsActionButtonVisible(), RepeatMode );
+}
+
+
 void FUICommandList::MapAction( const TSharedPtr< const FUICommandInfo > InUICommandInfo, FExecuteAction ExecuteAction, FCanExecuteAction CanExecuteAction, FIsActionChecked IsCheckedDelegate, FIsActionButtonVisible IsVisibleDelegate, EUIActionRepeatMode RepeatMode )
+{
+	MapAction( InUICommandInfo, ExecuteAction, CanExecuteAction, FGetActionCheckState::CreateStatic(&FUIAction::IsActionCheckedPassthrough, IsCheckedDelegate), IsVisibleDelegate, RepeatMode );
+}
+
+
+void FUICommandList::MapAction( const TSharedPtr< const FUICommandInfo > InUICommandInfo, FExecuteAction ExecuteAction, FCanExecuteAction CanExecuteAction, FGetActionCheckState GetActionCheckState, FIsActionButtonVisible IsVisibleDelegate, EUIActionRepeatMode RepeatMode )
 {
 	FUIAction Action;
 	Action.ExecuteAction = ExecuteAction;
 	Action.CanExecuteAction = CanExecuteAction;
-	Action.IsCheckedDelegate = IsCheckedDelegate;
+	Action.GetActionCheckState = GetActionCheckState;
 	Action.IsActionVisibleDelegate = IsVisibleDelegate;
 	Action.RepeatMode = RepeatMode;
 
@@ -119,16 +131,16 @@ EVisibility FUICommandList::GetVisibility( const TSharedRef< const FUICommandInf
 }
 
 
-bool FUICommandList::IsChecked( const TSharedRef< const FUICommandInfo > InUICommandInfo ) const
+ECheckBoxState FUICommandList::GetCheckState( const TSharedRef< const FUICommandInfo > InUICommandInfo ) const
 {
 	const FUIAction* Action = GetActionForCommand(InUICommandInfo);
 
 	if( Action )
 	{
-		return Action->IsChecked();
+		return Action->GetCheckState();
 	}
 
-	return false;
+	return ECheckBoxState::Unchecked;
 }
 
 
diff --git a/Engine/Source/Runtime/Slate/Private/Framework/MultiBox/SButtonRowBlock.cpp b/Engine/Source/Runtime/Slate/Private/Framework/MultiBox/SButtonRowBlock.cpp
--- a/Engine/Source/Runtime/Slate/Private/Framework/MultiBox/SButtonRowBlock.cpp
+++ b/Engine/Source/Runtime/Slate/Private/Framework/MultiBox/SButtonRowBlock.cpp
@@ -254,19 +254,18 @@ ECheckBoxState SButtonRowBlock::OnIsChecked() const
 	TSharedPtr< const FUICommandInfo > Action = MultiBlock->GetAction();
 	const FUIAction& DirectActions = MultiBlock->GetDirectActions();
 
-	bool bIsChecked = true;
+	ECheckBoxState CheckState = ECheckBoxState::Unchecked;
 	if( ActionList.IsValid() && Action.IsValid() )
 	{
-		bIsChecked = ActionList->IsChecked( Action.ToSharedRef() );
+		CheckState = ActionList->GetCheckState( Action.ToSharedRef() );
 	}
 	else
 	{
 		// There is no action list or action associated with this block via a UI command.  Execute any direct action we have
-		bIsChecked = DirectActions.IsChecked();
+		CheckState = DirectActions.GetCheckState();
 	}
 
-	return bIsChecked ? ECheckBoxState::Checked : ECheckBoxState::Unchecked;
-
+	return CheckState;
 }
 
 /**
diff --git a/Engine/Source/Runtime/Slate/Private/Framework/MultiBox/SMenuEntryBlock.cpp b/Engine/Source/Runtime/Slate/Private/Framework/MultiBox/SMenuEntryBlock.cpp
--- a/Engine/Source/Runtime/Slate/Private/Framework/MultiBox/SMenuEntryBlock.cpp
+++ b/Engine/Source/Runtime/Slate/Private/Framework/MultiBox/SMenuEntryBlock.cpp
@@ -390,6 +390,8 @@ TSharedRef< SWidget > SMenuEntryBlock::BuildMenuEntryWidget( const FMenuEntryBui
 				SNew( SBox )
 				.WidthOverride( MultiBoxConstants::MenuIconSize )
 				.HeightOverride( MultiBoxConstants::MenuIconSize )
+				.HAlign(HAlign_Center)
+				.VAlign(VAlign_Center)
 				[
 					IconWidget
 				]
@@ -971,18 +973,18 @@ ECheckBoxState SMenuEntryBlock::IsChecked() const
 	TSharedPtr< const FUICommandInfo > Action = MultiBlock->GetAction();
 	const FUIAction& DirectActions = MultiBlock->GetDirectActions();
 
-	bool bIsChecked = true;
+	ECheckBoxState CheckState = ECheckBoxState::Unchecked;
 	if( ActionList.IsValid() && Action.IsValid() )
 	{
-		bIsChecked = ActionList->IsChecked( Action.ToSharedRef() );
+		CheckState = ActionList->GetCheckState( Action.ToSharedRef() );
 	}
 	else
 	{
 		// There is no action list or action associated with this block via a UI command.  Execute any direct action we have
-		bIsChecked = DirectActions.IsChecked();
+		CheckState = DirectActions.GetCheckState();
 	}
 
-	return bIsChecked ? ECheckBoxState::Checked : ECheckBoxState::Unchecked;
+	return CheckState;
 }
 
 const FSlateBrush* SMenuEntryBlock::OnGetCheckImage() const
diff --git a/Engine/Source/Runtime/Slate/Private/Framework/MultiBox/SToolBarButtonBlock.cpp b/Engine/Source/Runtime/Slate/Private/Framework/MultiBox/SToolBarButtonBlock.cpp
--- a/Engine/Source/Runtime/Slate/Private/Framework/MultiBox/SToolBarButtonBlock.cpp
+++ b/Engine/Source/Runtime/Slate/Private/Framework/MultiBox/SToolBarButtonBlock.cpp
@@ -333,18 +333,18 @@ ECheckBoxState SToolBarButtonBlock::OnIsChecked() const
 	TSharedPtr< const FUICommandInfo > Action = MultiBlock->GetAction();
 	const FUIAction& DirectActions = MultiBlock->GetDirectActions();
 
-	bool bIsChecked = true;
+	ECheckBoxState CheckState = ECheckBoxState::Unchecked;
 	if( ActionList.IsValid() && Action.IsValid() )
 	{
-		bIsChecked = ActionList->IsChecked( Action.ToSharedRef() );
+		CheckState = ActionList->GetCheckState( Action.ToSharedRef() );
 	}
 	else
 	{
 		// There is no action list or action associated with this block via a UI command.  Execute any direct action we have
-		bIsChecked = DirectActions.IsChecked();
+		CheckState = DirectActions.GetCheckState();
 	}
 
-	return bIsChecked ? ECheckBoxState::Checked : ECheckBoxState::Unchecked;
+	return CheckState;
 }
 
 /**
diff --git a/Engine/Source/Runtime/Slate/Public/Framework/Commands/UIAction.h b/Engine/Source/Runtime/Slate/Public/Framework/Commands/UIAction.h
--- a/Engine/Source/Runtime/Slate/Public/Framework/Commands/UIAction.h
+++ b/Engine/Source/Runtime/Slate/Public/Framework/Commands/UIAction.h
@@ -12,6 +12,9 @@ DECLARE_DELEGATE_RetVal(bool, FCanExecuteAction);
 /** Defines FIsActionChecked delegate interface.  Returns true if the action is currently toggled on. */
 DECLARE_DELEGATE_RetVal(bool, FIsActionChecked);
 
+/** Defines FGetActionCheckState delegate interface.  Returns the ECheckBoxState for the action. */
+DECLARE_DELEGATE_RetVal(ECheckBoxState, FGetActionCheckState);
+
 /** Defines FIsActionButtonVisible delegate interface.  Returns true when UI buttons associated with the action should be visible. */
 DECLARE_DELEGATE_RetVal(bool, FIsActionButtonVisible);
 
@@ -35,8 +38,8 @@ struct FUIAction
 	/** Holds a delegate that is executed when determining whether this action can execute. */
 	FCanExecuteAction CanExecuteAction;
 
-	/** Holds a delegate that is executed when determining whether this action is checked. */
-	FIsActionChecked IsCheckedDelegate;
+	/** Holds a delegate that is executed when determining the check state of this action. */
+	FGetActionCheckState GetActionCheckState;
 
 	/** Holds a delegate that is executed when determining whether this action is visible. */
 	FIsActionButtonVisible IsActionVisibleDelegate;
@@ -50,7 +53,7 @@ public:
 	FUIAction( )
 		: ExecuteAction()
 		, CanExecuteAction()
-		, IsCheckedDelegate()
+		, GetActionCheckState()
 		, IsActionVisibleDelegate()
 		, RepeatMode(EUIActionRepeatMode::RepeatDisabled)
 	{ }
@@ -64,7 +67,7 @@ public:
 	FUIAction( FExecuteAction InitExecuteAction, EUIActionRepeatMode InitRepeatMode = EUIActionRepeatMode::RepeatDisabled )
 		: ExecuteAction(InitExecuteAction)
 		, CanExecuteAction(FCanExecuteAction::CreateRaw(&FSlateApplication::Get(), &FSlateApplication::IsNormalExecution))
-		, IsCheckedDelegate()
+		, GetActionCheckState()
 		, IsActionVisibleDelegate()
 		, RepeatMode(InitRepeatMode)
 	{ }
@@ -79,7 +82,7 @@ public:
 	FUIAction( FExecuteAction InitExecuteAction, FCanExecuteAction InitCanExecuteAction, EUIActionRepeatMode InitRepeatMode = EUIActionRepeatMode::RepeatDisabled )
 		: ExecuteAction(InitExecuteAction )
 		, CanExecuteAction(InitCanExecuteAction)
-		, IsCheckedDelegate()
+		, GetActionCheckState()
 		, IsActionVisibleDelegate()
 		, RepeatMode(InitRepeatMode)
 	{ }
@@ -95,7 +98,7 @@ public:
 	FUIAction( FExecuteAction InitExecuteAction, FCanExecuteAction InitCanExecuteAction, FIsActionChecked InitIsCheckedDelegate, EUIActionRepeatMode InitRepeatMode = EUIActionRepeatMode::RepeatDisabled )
 		: ExecuteAction(InitExecuteAction)
 		, CanExecuteAction(InitCanExecuteAction)
-		, IsCheckedDelegate(InitIsCheckedDelegate)
+		, GetActionCheckState(FGetActionCheckState::CreateStatic(&FUIAction::IsActionCheckedPassthrough, InitIsCheckedDelegate))
 		, IsActionVisibleDelegate()
 		, RepeatMode(InitRepeatMode)
 	{ }
@@ -103,16 +106,49 @@ public:
 	/**
 	 * Constructor that takes delegates to initialize the action with
 	 *
-	 * @param	ExecuteAction			The delegate to call when the action should be executed
-	 * @param	CanExecuteAction		The delegate to call to see if the action can be executed
-	 * @param	IsCheckedDelegate		The delegate to call to see if the action should appear checked when visualized
-	 * @param	IsActionVisibleDelegate	The delegate to call to see if the action should be visible
-	 * @param	RepeatMode				Can this action can be repeated if the chord used to call it is held down?
+	 * @param	ExecuteAction		The delegate to call when the action should be executed
+	 * @param	CanExecuteAction	The delegate to call to see if the action can be executed
+	 * @param	GetActionCheckState	The delegate to call to see what the check state of the action should be
+	 * @param	RepeatMode			Can this action can be repeated if the chord used to call it is held down?
+	 */
+	FUIAction( FExecuteAction InitExecuteAction, FCanExecuteAction InitCanExecuteAction, FGetActionCheckState InitGetActionCheckStateDelegate, EUIActionRepeatMode InitRepeatMode = EUIActionRepeatMode::RepeatDisabled )
+		: ExecuteAction(InitExecuteAction)
+		, CanExecuteAction(InitCanExecuteAction)
+		, GetActionCheckState(InitGetActionCheckStateDelegate)
+		, IsActionVisibleDelegate()
+		, RepeatMode(InitRepeatMode)
+	{ }
+
+	/**
+	 * Constructor that takes delegates to initialize the action with
+	 *
+	 * @param	ExecuteAction		The delegate to call when the action should be executed
+	 * @param	CanExecuteAction	The delegate to call to see if the action can be executed
+	 * @param	IsCheckedDelegate	The delegate to call to see if the action should appear checked when visualized
+	 * @param	IsActionVisible		The delegate to call to see if the action should be visible
+	 * @param	RepeatMode			Can this action can be repeated if the chord used to call it is held down?
 	 */
 	FUIAction( FExecuteAction InitExecuteAction, FCanExecuteAction InitCanExecuteAction, FIsActionChecked InitIsCheckedDelegate, FIsActionButtonVisible InitIsActionVisibleDelegate, EUIActionRepeatMode InitRepeatMode = EUIActionRepeatMode::RepeatDisabled )
 		: ExecuteAction(InitExecuteAction)
 		, CanExecuteAction(InitCanExecuteAction)
-		, IsCheckedDelegate(InitIsCheckedDelegate)
+		, GetActionCheckState(FGetActionCheckState::CreateStatic(&FUIAction::IsActionCheckedPassthrough, InitIsCheckedDelegate))
+		, IsActionVisibleDelegate(InitIsActionVisibleDelegate)
+		, RepeatMode(InitRepeatMode)
+	{ }
+
+	/**
+	 * Constructor that takes delegates to initialize the action with
+	 *
+	 * @param	ExecuteAction		The delegate to call when the action should be executed
+	 * @param	CanExecuteAction	The delegate to call to see if the action can be executed
+	 * @param	GetActionCheckState	The delegate to call to see what the check state of the action should be
+	 * @param	IsActionVisible		The delegate to call to see if the action should be visible
+	 * @param	RepeatMode			Can this action can be repeated if the chord used to call it is held down?
+	 */
+	FUIAction( FExecuteAction InitExecuteAction, FCanExecuteAction InitCanExecuteAction, FGetActionCheckState InitGetActionCheckStateDelegate, FIsActionButtonVisible InitIsActionVisibleDelegate, EUIActionRepeatMode InitRepeatMode = EUIActionRepeatMode::RepeatDisabled )
+		: ExecuteAction(InitExecuteAction)
+		, CanExecuteAction(InitCanExecuteAction)
+		, GetActionCheckState(InitGetActionCheckStateDelegate)
 		, IsActionVisibleDelegate(InitIsActionVisibleDelegate)
 		, RepeatMode(InitRepeatMode)
 	{ }
@@ -150,19 +186,14 @@ public:
 	 *
 	 * @return	Checked state for this action
 	 */
-	bool IsChecked( ) const
+	ECheckBoxState GetCheckState( ) const
 	{
-		if (IsCheckedDelegate.IsBound())
+		if (GetActionCheckState.IsBound())
 		{
-			const bool bIsChecked = IsCheckedDelegate.Execute();
-
-			if (bIsChecked)
-			{
-				return true;
-			}
+			return GetActionCheckState.Execute();
 		}
 
-		return false;
+		return ECheckBoxState::Unchecked;
 	}
 
 	/**
@@ -200,4 +231,18 @@ public:
 	{
 		return RepeatMode == EUIActionRepeatMode::RepeatEnabled;
 	}
+
+	/**
+	 * Passthrough function to convert the result from an FIsActionChecked delegate into something that works with a FGetActionCheckState delegate
+	 */
+	static ECheckBoxState IsActionCheckedPassthrough(FIsActionChecked InDelegate)
+	{
+		if (InDelegate.IsBound())
+		{
+			const bool bIsChecked = InDelegate.Execute();
+			return bIsChecked ? ECheckBoxState::Checked : ECheckBoxState::Unchecked;
+		}
+
+		return ECheckBoxState::Unchecked;
+	}
 };
diff --git a/Engine/Source/Runtime/Slate/Public/Framework/Commands/UICommandList.h b/Engine/Source/Runtime/Slate/Public/Framework/Commands/UICommandList.h
--- a/Engine/Source/Runtime/Slate/Public/Framework/Commands/UICommandList.h
+++ b/Engine/Source/Runtime/Slate/Public/Framework/Commands/UICommandList.h
@@ -41,6 +41,17 @@ public:
 	 */
 	void MapAction( const TSharedPtr< const FUICommandInfo > InUICommandInfo, FExecuteAction ExecuteAction, FCanExecuteAction CanExecuteAction, FIsActionChecked IsCheckedDelegate, EUIActionRepeatMode RepeatMode = EUIActionRepeatMode::RepeatDisabled );
 
+	/**
+	 * Maps a command info to a series of delegates that are executed by a multibox or mouse/keyboard input
+	 *
+	 * @param InUICommandInfo		The command info to map
+	 * @param ExecuteAction			The delegate to call when the command should be executed
+	 * @param CanExecuteAction		The delegate to call to see if the command can be executed
+	 * @param GetActionCheckState	The delegate to call to see what the check state of the action should be when visualized in a multibox
+	 * @param RepeatMode			Can this action can be repeated if the chord used to call it is held down?
+	 */
+	void MapAction( const TSharedPtr< const FUICommandInfo > InUICommandInfo, FExecuteAction ExecuteAction, FCanExecuteAction CanExecuteAction, FGetActionCheckState GetActionCheckState, EUIActionRepeatMode RepeatMode = EUIActionRepeatMode::RepeatDisabled );
+
 	/**
 	 * Maps a command info to a series of delegates that are executed by a multibox or mouse/keyboard input
 	 *
@@ -53,6 +64,18 @@ public:
 	 */
 	void MapAction( const TSharedPtr< const FUICommandInfo > InUICommandInfo, FExecuteAction ExecuteAction, FCanExecuteAction CanExecuteAction, FIsActionChecked IsCheckedDelegate, FIsActionButtonVisible IsVisibleDelegate, EUIActionRepeatMode RepeatMode = EUIActionRepeatMode::RepeatDisabled );
 
+	/**
+	 * Maps a command info to a series of delegates that are executed by a multibox or mouse/keyboard input
+	 *
+	 * @param InUICommandInfo		The command info to map
+	 * @param ExecuteAction			The delegate to call when the command should be executed
+	 * @param CanExecuteAction		The delegate to call to see if the command can be executed
+	 * @param GetActionCheckState	The delegate to call to see what the check state of the action should be when visualized in a multibox
+	 * @param IsVisibleDelegate		The delegate to call to see if the command should appear or be hidden when visualized in a multibox
+	 * @param RepeatMode			Can this action can be repeated if the chord used to call it is held down?
+	 */
+	void MapAction( const TSharedPtr< const FUICommandInfo > InUICommandInfo, FExecuteAction ExecuteAction, FCanExecuteAction CanExecuteAction, FGetActionCheckState GetActionCheckState, FIsActionButtonVisible IsVisibleDelegate, EUIActionRepeatMode RepeatMode = EUIActionRepeatMode::RepeatDisabled );
+
 	/**
 	 * Maps a command info to a series of delegates that are executed by a multibox or mouse/keyboard input
 	 *
@@ -97,11 +120,11 @@ public:
 	EVisibility GetVisibility( const TSharedRef< const FUICommandInfo > InUICommandInfo ) const;
 
 	/**
-	 * Calls the IsChecked delegate to see if the visualization of this command in a multibox should appear checked
+	 * Calls the GetCheckState delegate to see if the visualization of this command in a multibox should appear checked
 	 *
 	 * @param InUICommandInfo	The command info execute
 	 */
-	bool IsChecked( const TSharedRef< const FUICommandInfo > InUICommandInfo ) const;
+	ECheckBoxState GetCheckState( const TSharedRef< const FUICommandInfo > InUICommandInfo ) const;
 
 	/**
 	 * Processes any UI commands which are activated by the specified key event
diff --git a/Engine/Source/Runtime/Slate/Public/Widgets/Input/SCheckBox.h b/Engine/Source/Runtime/Slate/Public/Widgets/Input/SCheckBox.h
--- a/Engine/Source/Runtime/Slate/Public/Widgets/Input/SCheckBox.h
+++ b/Engine/Source/Runtime/Slate/Public/Widgets/Input/SCheckBox.h
@@ -4,19 +4,6 @@
 
 #include "SBorder.h"
 #include "SlateDelegates.h"
-#include "SCheckBox.generated.h"
-
-/** Current state of the check box */
-UENUM(BlueprintType)
-enum class ECheckBoxState : uint8
-{
-	/** Unchecked */
-	Unchecked,
-	/** Checked */
-	Checked,
-	/** Neither checked nor unchecked */
-	Undetermined
-};
 
 /** DEPRECATED 4.6 - Do not use */
 //@Todo slate: Remove this as soon as the 4.6 deprecated API is Removed.
diff --git a/Engine/Source/Runtime/SlateCore/Private/Styling/CoreStyle.cpp b/Engine/Source/Runtime/SlateCore/Private/Styling/CoreStyle.cpp
--- a/Engine/Source/Runtime/SlateCore/Private/Styling/CoreStyle.cpp
+++ b/Engine/Source/Runtime/SlateCore/Private/Styling/CoreStyle.cpp
@@ -759,12 +759,15 @@ TSharedRef<ISlateStyle> FCoreStyle::Create( const FName& InStyleSetName )
 
 		/* Create style for "ToolBar.CheckBox" widget ... */
 		const FCheckBoxStyle ToolBarCheckBoxStyle = FCheckBoxStyle()
-				.SetUncheckedImage( IMAGE_BRUSH( "Common/SmallCheckBox", Icon14x14 ) )
-				.SetCheckedImage( IMAGE_BRUSH( "Common/SmallCheckBox_Checked",  Icon14x14 ) )
-				.SetUncheckedHoveredImage( IMAGE_BRUSH( "Common/SmallCheckBox_Hovered", Icon14x14 ) )
-				.SetCheckedImage( IMAGE_BRUSH( "Common/SmallCheckBox_Checked_Hovered", Icon14x14 ) )
-				.SetUncheckedPressedImage(IMAGE_BRUSH( "Common/SmallCheckBox_Hovered", Icon14x14, FLinearColor( 0.5f, 0.5f, 0.5f ) ) )
-				.SetCheckedPressedImage( IMAGE_BRUSH( "Common/SmallCheckBox_Checked_Hovered", Icon14x14, FLinearColor( 0.5f, 0.5f, 0.5f ) ) );
+			.SetUncheckedImage( IMAGE_BRUSH( "Common/SmallCheckBox", Icon14x14 ) )
+			.SetCheckedImage( IMAGE_BRUSH( "Common/SmallCheckBox_Checked",  Icon14x14 ) )
+			.SetUncheckedHoveredImage( IMAGE_BRUSH( "Common/SmallCheckBox_Hovered", Icon14x14 ) )
+			.SetCheckedImage( IMAGE_BRUSH( "Common/SmallCheckBox_Checked_Hovered", Icon14x14 ) )
+			.SetUncheckedPressedImage(IMAGE_BRUSH( "Common/SmallCheckBox_Hovered", Icon14x14, FLinearColor( 0.5f, 0.5f, 0.5f ) ) )
+			.SetCheckedPressedImage( IMAGE_BRUSH( "Common/SmallCheckBox_Checked_Hovered", Icon14x14, FLinearColor( 0.5f, 0.5f, 0.5f ) ) )
+			.SetUndeterminedImage( IMAGE_BRUSH( "Common/CheckBox_Undetermined", Icon14x14 ) )
+			.SetUndeterminedHoveredImage( IMAGE_BRUSH( "Common/CheckBox_Undetermined_Hovered", Icon14x14 ) )
+			.SetUndeterminedPressedImage( IMAGE_BRUSH( "Common/CheckBox_Undetermined_Hovered", Icon14x14, FLinearColor( 0.5f, 0.5f, 0.5f ) ) );
 		/* ... and add new style */
 		Style->Set( "ToolBar.CheckBox", ToolBarCheckBoxStyle );
 
@@ -841,9 +844,9 @@ TSharedRef<ISlateStyle> FCoreStyle::Create( const FName& InStyleSetName )
 			.SetCheckedImage( IMAGE_BRUSH( "Common/SmallCheckBox_Checked", Icon14x14 ) )
 			.SetCheckedHoveredImage( IMAGE_BRUSH( "Common/SmallCheckBox_Checked_Hovered", Icon14x14 ) )
 			.SetCheckedPressedImage( IMAGE_BRUSH( "Common/SmallCheckBox_Checked_Hovered", Icon14x14, FLinearColor( 0.5f, 0.5f, 0.5f ) ) )
-			.SetUndeterminedImage( IMAGE_BRUSH( "Common/CheckBox_Undetermined", Icon16x16 ) )
-			.SetUndeterminedHoveredImage( IMAGE_BRUSH( "Common/CheckBox_Checked_Hovered", Icon16x16 ) )
-			.SetUndeterminedPressedImage( IMAGE_BRUSH( "Common/CheckBox_Undetermined_Hovered", Icon16x16, FLinearColor( 0.5f, 0.5f, 0.5f ) ) );
+			.SetUndeterminedImage( IMAGE_BRUSH( "Common/CheckBox_Undetermined", Icon14x14 ) )
+			.SetUndeterminedHoveredImage( IMAGE_BRUSH( "Common/CheckBox_Undetermined_Hovered", Icon14x14 ) )
+			.SetUndeterminedPressedImage( IMAGE_BRUSH( "Common/CheckBox_Undetermined_Hovered", Icon14x14, FLinearColor( 0.5f, 0.5f, 0.5f ) ) );
  
 		/* ...and add the new style */
 		Style->Set( "Menu.CheckBox", BasicMenuCheckBoxStyle );
diff --git a/Engine/Source/Runtime/SlateCore/Public/Styling/SlateTypes.h b/Engine/Source/Runtime/SlateCore/Public/Styling/SlateTypes.h
--- a/Engine/Source/Runtime/SlateCore/Public/Styling/SlateTypes.h
+++ b/Engine/Source/Runtime/SlateCore/Public/Styling/SlateTypes.h
@@ -35,6 +35,18 @@ namespace ESlateCheckBoxType
 	};
 }
 
+/** Current state of the check box */
+UENUM(BlueprintType)
+enum class ECheckBoxState : uint8
+{
+	/** Unchecked */
+	Unchecked,
+	/** Checked */
+	Checked,
+	/** Neither checked nor unchecked */
+	Undetermined
+};
+
 /**
  * Represents the appearance of an SCheckBox
  */
