Profiler - Removed obsolete memory profiler code from the profiler session

diff --git a/Engine/Source/Developer/Profiler/Private/ProfilerSession.cpp b/Engine/Source/Developer/Profiler/Private/ProfilerSession.cpp
--- a/Engine/Source/Developer/Profiler/Private/ProfilerSession.cpp
+++ b/Engine/Source/Developer/Profiler/Private/ProfilerSession.cpp
@@ -549,6 +549,7 @@ static double GetSecondsPerCycle( const FStatPacketArray& Frame )
 				check( Item.NameAndInfo.GetFlag( EStatMetaFlags::DummyAlwaysOne ) ); 
 
 				const FName LongName = Item.NameAndInfo.GetEncodedName();
+				const FName RawName = Item.NameAndInfo.GetRawName();
 				if( LongName.IsEqual( SecondsPerCycleFName, ENameCase::IgnoreCase, false ) )
 				{
 					Result = Item.GetValue_double();
@@ -733,29 +734,23 @@ void FRawProfilerSession::PrepareLoading()
 		Frames.Sort();
 		const int64 MiddleFrame = Frames[Frames.Num()/2];
 
-		const bool bProcessMemory = false;
-		if( bProcessMemory )
-		{
-			ProcessMemoryOperations(CombinedHistory);
-		}
-		else
+
+		// Remove all frames without the game thread messages.
+		for (int32 FrameIndex = 0; FrameIndex < Frames.Num(); ++FrameIndex)
 		{
-			// Remove all frames without the game thread messages.
-			for( int32 FrameIndex = 0; FrameIndex < Frames.Num(); ++FrameIndex )
-			{
-				const int64 TargetFrame = Frames[FrameIndex];
-				const FStatPacketArray& Frame = CombinedHistory.FindChecked(TargetFrame);
+			const int64 TargetFrame = Frames[FrameIndex];
+			const FStatPacketArray& Frame = CombinedHistory.FindChecked( TargetFrame );
 
-				const double GameThreadTimeMS = GetMetaData()->ConvertCyclesToMS( GetFastThreadFrameTimeInternal(Frame,EThreadType::Game) );
+			const double GameThreadTimeMS = GetMetaData()->ConvertCyclesToMS( GetFastThreadFrameTimeInternal( Frame, EThreadType::Game ) );
 
-				if( GameThreadTimeMS == 0.0f )
-				{
-					CombinedHistory.Remove(TargetFrame);
-					Frames.RemoveAt(FrameIndex);
-					FrameIndex--;
-				}
+			if (GameThreadTimeMS == 0.0f)
+			{
+				CombinedHistory.Remove( TargetFrame );
+				Frames.RemoveAt( FrameIndex );
+				FrameIndex--;
 			}
 		}
+		
 	
 		StatMetaData->SecondsPerCycle = GetSecondsPerCycle( CombinedHistory.FindChecked(MiddleFrame) );
 		check( StatMetaData->GetSecondsPerCycle() > 0.0 );
@@ -1193,564 +1188,4 @@ void FRawProfilerSession::ProcessStatPacketArray( const FStatPacketArray& StatPa
 	}
 
 	out_ProfilerFrame.SortChildren();
-}
-
-
-//-----------------------------------------------------------------------------
-
-const TCHAR* CallstackSeparator = TEXT("+");
-
-static FString EncodeCallstack( const TArray<FName>& Callstack )
-{
-	FString Result;
-	for( const auto& Name : Callstack )
-	{
-		Result += TTypeToString<int32>::ToString((int32)Name.GetComparisonIndex());
-		Result += CallstackSeparator;
-	}
-	return Result;
-}
-
-static void DecodeCallstack( const FName& EncodedCallstack, TArray<FString>& out_DecodedCallstack )
-{
-	EncodedCallstack.ToString().ParseIntoArray( out_DecodedCallstack, CallstackSeparator, true );
-}
-
-/** Simple allocation info. Assumes the sequence tag are unique and doesn't wrap. */
-struct FAllocationInfo
-{
-	uint64 Ptr;
-	int64 Size;
-	uint32 SequenceTag;
-	FName EncodedCallstack;
-	EMemoryOperation Op;
-	bool bHasBrokenCallstack;
-	
-	FAllocationInfo(
-		uint64 InPtr, 
-		int64 InSize, 
-		const TArray<FName>& InCallstack, 
-		uint32 InSequenceTag, 
-		EMemoryOperation InOp, 
-		bool bInHasBrokenCallstack)
-		: Ptr(InPtr)
-		, Size(InSize)
-		, SequenceTag(InSequenceTag)
-		, EncodedCallstack(*EncodeCallstack(InCallstack))
-		, Op(InOp)
-		, bHasBrokenCallstack(bInHasBrokenCallstack)
-	{
-	}
-
-	FAllocationInfo(const FAllocationInfo& Other)
-		: Ptr(Other.Ptr)
-		, Size(Other.Size)
-		, SequenceTag(Other.SequenceTag)
-		, EncodedCallstack(Other.EncodedCallstack)
-		, Op(Other.Op)
-		, bHasBrokenCallstack(Other.bHasBrokenCallstack)
-	{}
-
-	FAllocationInfo()
-		: Ptr(0)
-		, Size(0)
-		, SequenceTag(0)
-		, Op(EMemoryOperation::Invalid)
-	{}
-
-
-	bool operator<( const FAllocationInfo& Other ) const
-	{
-		return SequenceTag < Other.SequenceTag;
-	}
-
-	bool operator==( const FAllocationInfo& Other ) const
-	{
-		return SequenceTag == Other.SequenceTag;
-	}
-};
-
-struct FAllocationInfoGreater
-{
-	FORCEINLINE bool operator()(const FAllocationInfo& A, const FAllocationInfo& B) const { return B.Size < A.Size; }
-};
-
-
-struct FSizeAndCount
-{
-	uint64 Size;
-	uint64 Count;
-	FSizeAndCount()
-		: Size(0)
-		, Count(0)
-	{}
-};
-
-struct FSizeAndCountGreater
-{
-	FORCEINLINE bool operator()(const FSizeAndCount& A, const FSizeAndCount& B) const { return B.Size < A.Size; }
-};
-
-/** Holds stats stack state, used to preserve continuity when the game frame has changed. */
-struct FStackState
-{
-	FStackState()
-		: bIsBrokenCallstack(false)
-	{}
-
-	/** Call stack. */
-	TArray<FName> Stack;
-
-	/** Current function name. */
-	FName Current;
-
-	/** Whether this callstack is marked as broken due to mismatched start and end scope cycles. */
-	bool bIsBrokenCallstack;
-};
-
-FString GetCallstack( const FName& EncodedCallstack )
-{
-	FString Result;
-
-	TArray<FString> DecodedCallstack;
-	DecodeCallstack( EncodedCallstack, DecodedCallstack );
-
-	for( int32 Index = DecodedCallstack.Num() - 1; Index >= 0; --Index )
-	{
-		NAME_INDEX NameIndex = 0;
-		TTypeFromString<NAME_INDEX>::FromString( NameIndex, *DecodedCallstack[Index] );
-
-		const FName LongName = FName(NameIndex, NameIndex, 0);
-
-		const FString ShortName = FStatNameAndInfo::GetShortNameFrom(LongName).ToString();
-		const FString Group = FStatNameAndInfo::GetGroupNameFrom(LongName).ToString();
-		FString Desc = FStatNameAndInfo::GetDescriptionFrom(LongName);
-		Desc.Trim();
-
-		const bool bHasDesc = Desc.Len() > 0;
-		if( Desc != ShortName )
-		{
-			if( bHasDesc )
-			{
-				Desc += TEXT( " [" );
-			}
-			Desc += ShortName;
-			if( bHasDesc )
-			{
-				Desc += TEXT( "]" );
-			}
-		}
-
-		//Result += FString::Printf( TEXT("%s [%s] -> "), *Desc, *Group );
-		Result += Desc;
-		if( Index > 0 )
-		{
-			Result += TEXT(" <- ");
-		}
-	}
-
-	return Result;
-}
-
-// New version.
-void FRawProfilerSession::ProcessMemoryOperations( const TMap<int64, FStatPacketArray>& CombinedHistory )
-{
-	// This is only example code, no fully implemented, may sometimes crash.
-	// This code is not optimized. 
-
-	uint64 NumMemoryOperations = 0;
-
-	// Generate frames
-	TArray<int64> Frames;
-	CombinedHistory.GenerateKeyArray(Frames);
-	Frames.Sort();
-
-	const FProfilerStatMetaDataRef MetaData = GetMetaData();
-
-	// Raw stats callstack for this stat packet array.
-	TMap<FName,FStackState> StackStates;
-
-	// All allocation ordered by the sequence tag.
-	// There is an assumption that the sequence tag will not turn-around.
-	//TMap<uint32, FAllocationInfo> SequenceAllocationMap;
-	TArray<FAllocationInfo> SequenceAllocationArray;
-
-	// Pass 1.
-	// Read all stats messages, parse all memory operations and decode callstacks.
-	const int64 FirstFrame = 0;
-	for( int32 FrameIndex = 0; FrameIndex < Frames.Num(); ++FrameIndex )
-	{
-		UE_LOG( LogStats, Log, TEXT( "Processing %i/%i frame" ), FrameIndex, Frames.Num() );
-
-		const int64 TargetFrame = Frames[FrameIndex];
-		const int64 Diff = TargetFrame-FirstFrame;
-		const FStatPacketArray& Frame = CombinedHistory.FindChecked(TargetFrame);
-
-		for( int32 PacketIndex = 0; PacketIndex < Frame.Packets.Num(); PacketIndex++ )
-		{
-			UE_LOG( LogStats, Log, TEXT( "Processing packet %i/%i" ), PacketIndex, Frame.Packets.Num() );
-
-			const FStatPacket& StatPacket = *Frame.Packets[PacketIndex];
-			const FName& ThreadFName = StatsThreadStats.Threads.FindChecked( StatPacket.ThreadId );
-			const uint32 NewThreadID = MetaData->ThreadIDtoStatID.FindChecked( StatPacket.ThreadId );
-
-			FStackState* StackState = StackStates.Find(ThreadFName);
-			if( !StackState )
-			{
-				StackState = &StackStates.Add(ThreadFName);
-				StackState->Stack.Add( ThreadFName );
-				StackState->Current = ThreadFName;
-			}
-
-			const FStatMessagesArray& Data = StatPacket.StatMessages;
-
-			int32 LastPct = 0;
-			for( int32 Index = 0; Index < Data.Num(); Index++ )
-			{
-				const int32 CurrentPct = int32(100.0 * Index / Data.Num());
-				if( CurrentPct > LastPct+10 )
-				{
-					UE_LOG( LogStats, Log, TEXT( "Processing %3i %% (%i/%i) stat messages" ), CurrentPct, Index, Data.Num() );
-					LastPct = CurrentPct;
-				}
-				
-				const FStatMessage& Item = Data[Index];
-
-				const EStatOperation::Type Op = Item.NameAndInfo.GetField<EStatOperation>();
-				const FName LongName = Item.NameAndInfo.GetEncodedName();
-				const FName RawName = Item.NameAndInfo.GetRawName();
-
-				if( Op == EStatOperation::CycleScopeStart || Op == EStatOperation::CycleScopeEnd || Op == EStatOperation::Memory )
-				{
-					if( Op == EStatOperation::CycleScopeStart )
-					{
-						StackState->Stack.Add( RawName );
-						StackState->Current = RawName;
-					}
-					else if( Op == EStatOperation::Memory )
-					{
-						// Experimental code used only to test the implementation.
-						// First memory operation is Alloc or Free
-						const uint64 EncodedPtr = Item.GetValue_Ptr();
-						const bool bIsAlloc = (EncodedPtr & (uint64)EMemoryOperation::Alloc) != 0;
-						const bool bIsFree = (EncodedPtr & (uint64)EMemoryOperation::Free) != 0;
-						const uint64 Ptr = EncodedPtr & ~(uint64)EMemoryOperation::Mask;
-						if( bIsAlloc )
-						{
-							NumMemoryOperations++;
-							// @see FStatsMallocProfilerProxy::TrackAlloc
-							// After alloc ptr message there is always alloc size message and the sequence tag.
-							Index++;
-							const FStatMessage& AllocSizeMessage = Data[Index];
-							const int64 AllocSize = AllocSizeMessage.GetValue_int64();
-
-							// Read operation sequence tag.
-							Index++;
-							const FStatMessage& SequenceTagMessage = Data[Index];
-							const uint32 SequenceTag = SequenceTagMessage.GetValue_int64();
-
-							// Create a callstack.
-							TArray<FName> StatsBasedCallstack;
-							for( const auto& StackRawName : StackState->Stack )
-							{
-								StatsBasedCallstack.Add(StackRawName);
-							}
-
-							// Add a new allocation.
-							SequenceAllocationArray.Add( 
-								FAllocationInfo(
-									Ptr,
-									AllocSize, 
-									StatsBasedCallstack, 
-									SequenceTag, 
-									EMemoryOperation::Alloc, 						
-									StackState->bIsBrokenCallstack) );
-						}
-						else if( bIsFree )
-						{
-							NumMemoryOperations++;
-							// Read operation sequence tag.
-							Index++;
-							const FStatMessage& SequenceTagMessage = Data[Index];
-							const uint32 SequenceTag = SequenceTagMessage.GetValue_int64();
-
-							// Create a callstack.
-							TArray<FName> StatsBasedCallstack;
-							for( const auto& StackRawName : StackState->Stack )
-							{
-								StatsBasedCallstack.Add(StackRawName);
-							}
-
-							// Add a new free.
-							SequenceAllocationArray.Add( 
-								FAllocationInfo(
-									Ptr,
-									0, 
-									StatsBasedCallstack, 
-									SequenceTag,
-									EMemoryOperation::Free, 					
-									StackState->bIsBrokenCallstack) );
-						}
-						else
-						{
-							UE_LOG(LogStats, Warning, TEXT("Pointer from a memory operation is invalid") );
-						}
-					}
-					else if( Op == EStatOperation::CycleScopeEnd )
-					{
-						if( StackState->Stack.Num() > 1 )
-						{
-							const FName ScopeStart = StackState->Stack.Pop();
-							const FName ScopeEnd = Item.NameAndInfo.GetRawName();
-
-							check(ScopeStart==ScopeEnd);
-
-							StackState->Current = StackState->Stack.Last();
-
-							// The stack should be ok, but it may be partially broken.
-							// This will happen if memory profiling starts in the middle of executing a background thread.
-							StackState->bIsBrokenCallstack = false;
-						}
-						else
-						{
-							const FName ShortName = Item.NameAndInfo.GetShortName();
-
-							UE_LOG(LogStats, Warning, TEXT("Broken cycle scope end %s/%s, current %s"), 
-								*ThreadFName.ToString(), 
-								*ShortName.ToString(),
-								*StackState->Current.ToString() );
-
-							// The stack is completely broken, only has the thread name and the last cycle scope.
-							// Rollback to the thread node.
-							StackState->bIsBrokenCallstack = true;
-							StackState->Stack.Empty();
-							StackState->Stack.Add( ThreadFName );
-							StackState->Current = ThreadFName;
-						}
-					}
-				}
-			}
-		}
-	}
-
-	UE_LOG(LogStats, Warning, TEXT("NumMemoryOperations %llu"), NumMemoryOperations );
-	UE_LOG(LogStats, Warning, TEXT("SequenceAllocationArray.Num %i"), SequenceAllocationArray.Num() );
-
-	// Pass 2.
-	/*
-	TMap<uint32,FAllocationInfo> UniqueSeq;
-	TMultiMap<uint32,FAllocationInfo> OriginalAllocs;
-	TMultiMap<uint32,FAllocationInfo> BrokenAllocs;
-	for( const FAllocationInfo& Alloc : SequenceAllocationArray )
-	{
-		const FAllocationInfo* Found = UniqueSeq.Find(Alloc.SequenceTag);
-		if( !Found )
-		{
-			UniqueSeq.Add(Alloc.SequenceTag,Alloc);
-		}
-		else
-		{
-			OriginalAllocs.Add(Alloc.SequenceTag, *Found);
-			BrokenAllocs.Add(Alloc.SequenceTag, Alloc);
-		}		
-	}
-	*/
-
-	// Sort all memory operation by the sequence tag, iterate through all operation and generate memory usage.
-	SequenceAllocationArray.Sort( TLess<FAllocationInfo>() );
-
-	// Alive allocations.
-	TMap<uint64, FAllocationInfo> AllocationMap;
-	TMultiMap<uint64, FAllocationInfo> FreeWithoutAllocMap;
-	TMultiMap<uint64, FAllocationInfo> DuplicatedAllocMap;
-	int32 NumDuplicatedMemoryOperations = 0;
-	int32 NumFWAMemoryOperations = 0; // FreeWithoutAlloc
-
-	for( const FAllocationInfo& Alloc : SequenceAllocationArray )
-	{
-		const EMemoryOperation MemOp = Alloc.Op;
-		const uint64 Ptr = Alloc.Ptr;
-		const int64 Size = Alloc.Size;
-		const uint32 SequenceTag = Alloc.SequenceTag;
-
-		if( MemOp == EMemoryOperation::Alloc )
-		{
-			const FAllocationInfo* Found = AllocationMap.Find(Ptr);
-
-			if( !Found )
-			{
-				AllocationMap.Add(Ptr, Alloc);
-			}
-			else
-			{
-				const FAllocationInfo* FoundAndFreed = FreeWithoutAllocMap.Find(Found->Ptr);
-				const FAllocationInfo* FoundAndAllocated = FreeWithoutAllocMap.Find(Alloc.Ptr);
-
-				if( FoundAndFreed )
-				{
-					const FString FoundAndFreedCallstack = GetCallstack(FoundAndFreed->EncodedCallstack);
-				}
-
-				if( FoundAndAllocated )
-				{
-					const FString FoundAndAllocatedCallstack = GetCallstack(FoundAndAllocated->EncodedCallstack);
-				}
-
-				NumDuplicatedMemoryOperations++;
-
-				const FString FoundCallstack = GetCallstack(Found->EncodedCallstack);
-				const FString AllocCallstack = GetCallstack(Alloc.EncodedCallstack);
-
-				// Replace pointer.
-				AllocationMap.Add(Ptr, Alloc);
-				// Store the old pointer.
-				DuplicatedAllocMap.Add(Ptr, *Found);
-			}
-		}
-		else if( MemOp == EMemoryOperation::Free )
-		{
-			const FAllocationInfo* Found = AllocationMap.Find(Ptr);	
-			if( Found )
-			{
-				const bool bIsValid = Alloc.SequenceTag > Found->SequenceTag;
-				AllocationMap.Remove(Ptr);
-			}
-			else
-			{
-				FreeWithoutAllocMap.Add(Ptr,Alloc);
-				NumFWAMemoryOperations++;
-			}
-		}
-	}
-
-	UE_LOG(LogStats, Warning, TEXT("NumDuplicatedMemoryOperations %i"), NumDuplicatedMemoryOperations );
-	UE_LOG(LogStats, Warning, TEXT("NumFWAMemoryOperations %i"), NumFWAMemoryOperations );
-
-	// Dump problematic allocations
-	DuplicatedAllocMap.ValueSort( FAllocationInfoGreater() );
-	//FreeWithoutAllocMap
-
-	uint64 TotalDuplicatedMemory = 0;
-	for( const auto& It : DuplicatedAllocMap )
-	{
-		const FAllocationInfo& Alloc = It.Value;
-		TotalDuplicatedMemory += Alloc.Size;
-	}
-
-	const float MaxPctDisplayed = 0.80f;
-	uint64 DisplayedSoFar = 0;
-	for( const auto& It : DuplicatedAllocMap )
-	{
-		const FAllocationInfo& Alloc = It.Value;
-		const FString AllocCallstack = GetCallstack(Alloc.EncodedCallstack);
-		UE_LOG(LogStats, Warning, TEXT("%lli (%.2f MB) %s"), Alloc.Size, Alloc.Size/1024.0f/1024.0f, *AllocCallstack );
-
-		DisplayedSoFar += Alloc.Size;
-
-		const float CurrentPct = (float)DisplayedSoFar/(float)TotalDuplicatedMemory;
-		if( CurrentPct > MaxPctDisplayed )
-		{
-			break;
-		}
-	}
-
-	GenerateMemoryUsageReport( AllocationMap );
-}
-
-void FRawProfilerSession::GenerateMemoryUsageReport( const TMap<uint64, FAllocationInfo>& AllocationMap )
-{
-	if( AllocationMap.Num() == 0  )
-	{
-		UE_LOG(LogStats, Warning, TEXT("There are no allocations, make sure memory profiler is enabled") );
-	}
-	else
-	{
-		TMap<FName,FSizeAndCount> ScopedAllocations;
-
-		FName BrokenCallstackName = TEXT("245+12040+12265+14018+");
-		TArray<FAllocationInfo> BrokenAllocations;
-
-		uint64 TotalAllocatedMemory = 0;
-		for( const auto& It : AllocationMap )
-		{
-			const FAllocationInfo& Alloc = It.Value;
-			FSizeAndCount& SizeAndCount = ScopedAllocations.FindOrAdd(Alloc.EncodedCallstack);
-			SizeAndCount.Size += Alloc.Size;
-			SizeAndCount.Count += 1;
-
-			TotalAllocatedMemory += Alloc.Size;
-
-			if( Alloc.EncodedCallstack == BrokenCallstackName )
-			{
-				BrokenAllocations.Add(Alloc);
-			}
-		}
-
-		// Dump memory to the log.
-	
-		ScopedAllocations.ValueSort( FSizeAndCountGreater() );
-		BrokenAllocations.Sort( FAllocationInfoGreater() );
-
-		const float MaxPctDisplayed = 0.90f;
-		int32 CurrentIndex = 0;
-		uint64 DisplayedSoFar = 0;
-		UE_LOG( LogStats, Warning, TEXT("Index, Size (Size MB), Count, Stat desc") );
-		for( const auto& It : ScopedAllocations )
-		{
-			const FSizeAndCount& SizeAndCount = It.Value;
-			const FName& EncodedCallstack = It.Key;
-
-			const FString AllocCallstack = GetCallstack(EncodedCallstack);
-
-			/*TArray<FString> DecodedCallstack;
-			DecodeCallstack( EncodedCallstack, DecodedCallstack );
-
-			for( const auto& StrNameIndex : DecodedCallstack )
-			{
-				NAME_INDEX NameIndex = 0;
-				TTypeFromString<NAME_INDEX>::FromString( NameIndex, *StrNameIndex );
-
-				const FName LongName = FName(NameIndex, NameIndex, 0);
-
-				const FString ShortName = FStatNameAndInfo::GetShortNameFrom(LongName).ToString();
-				const FString Group = FStatNameAndInfo::GetGroupNameFrom(LongName).ToString();
-				FString Desc = FStatNameAndInfo::GetDescriptionFrom(LongName);
-				Desc.Trim();
-
-				if( Desc != ShortName )
-				{
-					if( Desc.Len() )
-					{
-						Desc += TEXT( " - " );
-					}
-					Desc += ShortName;
-				}
-
-				UE_LOG( LogStats, Warning, TEXT("%s [%s] "), 
-					*Desc, 
-					*Group );
-			}
-			}*/
-
-			UE_LOG( LogStats, Warning, TEXT("%2i, %llu (%.2f MB), %llu, %s"), 
-				CurrentIndex, 
-				SizeAndCount.Size, 
-				SizeAndCount.Size/1024.0f/1024.0f,
-				SizeAndCount.Count, 
-				*AllocCallstack );
-
-			CurrentIndex++;
-			DisplayedSoFar += SizeAndCount.Size;
-
-			const float CurrentPct = (float)DisplayedSoFar/(float)TotalAllocatedMemory;
-			if( CurrentPct > MaxPctDisplayed )
-			{
-				break;
-			}
-		}
-
-		UE_LOG(LogStats, Warning, TEXT("Allocated memory: %llu bytes (%.2f MB)"), TotalAllocatedMemory, TotalAllocatedMemory/1024.0f/1024.0f );
-	}
-
-	FPlatformMisc::RequestExit(true);
 }
\ No newline at end of file
diff --git a/Engine/Source/Developer/Profiler/Private/ProfilerSession.h b/Engine/Source/Developer/Profiler/Private/ProfilerSession.h
--- a/Engine/Source/Developer/Profiler/Private/ProfilerSession.h
+++ b/Engine/Source/Developer/Profiler/Private/ProfilerSession.h
@@ -1243,12 +1243,6 @@ class FRawProfilerSession : public FProfilerSession
 	/** Index of the last processed data for the mini-view. */
 	int32 CurrentMiniViewFrame;
 
-	/** Basic memory profiling, only for debugging purpose. */
-	void ProcessMemoryOperations( const TMap<int64, FStatPacketArray>& CombinedHistory );
-
-	/** Generate a basic memory usage report and prints it to the log. */
-	void GenerateMemoryUsageReport( const TMap<uint64, FAllocationInfo>& AllocationMap );
-
 public:
 	/**
 	* Default constructor, creates a profiler session from a capture file.
