Profiler - Game thread graph was not always enabled after loading a ue4stats file

diff --git a/Engine/Source/Developer/Profiler/Private/ProfilerDataSource.h b/Engine/Source/Developer/Profiler/Private/ProfilerDataSource.h
--- a/Engine/Source/Developer/Profiler/Private/ProfilerDataSource.h
+++ b/Engine/Source/Developer/Profiler/Private/ProfilerDataSource.h
@@ -967,7 +967,7 @@ class FEventGraphSample : public TSharedFromThis<FEventGraphSample>
 	enum
 	{
 		/** Maximum number of stack when traversing event graph - should be enough to store root nodes and all its children.*/
-		MaxStackSize = 65536
+		MaxStackSize = 256*1024
 	};
 
 	struct FDuplicateHierarchyTag {};
diff --git a/Engine/Source/Developer/Profiler/Private/ProfilerManager.cpp b/Engine/Source/Developer/Profiler/Private/ProfilerManager.cpp
--- a/Engine/Source/Developer/Profiler/Private/ProfilerManager.cpp
+++ b/Engine/Source/Developer/Profiler/Private/ProfilerManager.cpp
@@ -267,6 +267,8 @@ void FProfilerManager::ProfilerSession_OnCaptureFileProcessed( const FGuid Profi
 	const FProfilerSessionRef* ProfilerSession = FindSessionInstance( ProfilerInstanceID );
 	if( ProfilerSession && ProfilerWindow.IsValid())
 	{
+		TrackDefaultStats();
+
 		RequestFilterAndPresetsUpdateEvent.Broadcast();
 
 		GetProfilerWindow()->UpdateEventGraph( ProfilerInstanceID, (*ProfilerSession)->GetEventGraphDataAverage(), (*ProfilerSession)->GetEventGraphDataMaximum(), true );
