Profiler - Initial support for updated raw stats

diff --git a/Engine/Source/Developer/Profiler/Private/ProfilerPrivatePCH.h b/Engine/Source/Developer/Profiler/Private/ProfilerPrivatePCH.h
--- a/Engine/Source/Developer/Profiler/Private/ProfilerPrivatePCH.h
+++ b/Engine/Source/Developer/Profiler/Private/ProfilerPrivatePCH.h
@@ -23,6 +23,7 @@
 // TODO: Remove dependencies between generic profiler code and profiler ui code
 
 // Common headers and managers.
+#include "ProfilerRawStats.h"
 #include "ProfilerStream.h"
 
 #include "ProfilerSample.h"
diff --git a/Engine/Source/Developer/Profiler/Private/ProfilerRawStats.cpp b/Engine/Source/Developer/Profiler/Private/ProfilerRawStats.cpp
new file mode 100644
--- /dev/null
+++ b/Engine/Source/Developer/Profiler/Private/ProfilerRawStats.cpp
@@ -0,0 +1,208 @@
+// Copyright 1998-2015 Epic Games, Inc. All Rights Reserved.
+
+#include "ProfilerPrivatePCH.h"
+
+
+/*-----------------------------------------------------------------------------
+	Sort helpers
+-----------------------------------------------------------------------------*/
+
+/** Sorts allocations by size. */
+struct FAllocationInfoSequenceTagLess
+{
+	FORCEINLINE bool operator()( const FAllocationInfo& A, const FAllocationInfo& B ) const
+	{
+		return A.SequenceTag < B.SequenceTag;
+	}
+};
+
+/** Sorts allocations by size. */
+struct FAllocationInfoSizeGreater
+{
+	FORCEINLINE bool operator()( const FAllocationInfo& A, const FAllocationInfo& B ) const
+	{
+		return B.Size < A.Size;
+	}
+};
+
+/** Sorts combined allocations by size. */
+struct FCombinedAllocationInfoSizeGreater
+{
+	FORCEINLINE bool operator()( const FCombinedAllocationInfo& A, const FCombinedAllocationInfo& B ) const
+	{
+		return B.Size < A.Size;
+	}
+};
+
+
+/** Sorts node allocations by size. */
+struct FNodeAllocationInfoSizeGreater
+{
+	FORCEINLINE bool operator()( const FNodeAllocationInfo& A, const FNodeAllocationInfo& B ) const
+	{
+		return B.Size < A.Size;
+	}
+};
+
+/*-----------------------------------------------------------------------------
+	Callstack decoding/encoding
+-----------------------------------------------------------------------------*/
+
+/** Helper struct used to manipulate stats based callstacks. */
+struct FStatsCallstack
+{
+	/** Separator. */
+	static const TCHAR* CallstackSeparator;
+
+	/** Encodes decoded callstack a string, to be like '45+656+6565'. */
+	static FString Encode( const TArray<FName>& Callstack )
+	{
+		FString Result;
+		for (const auto& Name : Callstack)
+		{
+			Result += TTypeToString<int32>::ToString( (int32)Name.GetComparisonIndex() );
+			Result += CallstackSeparator;
+		}
+		return Result;
+	}
+
+	/** Decodes encoded callstack to an array of FNames. */
+	static void DecodeToNames( const FName& EncodedCallstack, TArray<FName>& out_DecodedCallstack )
+	{
+		TArray<FString> DecodedCallstack;
+		DecodeToStrings( EncodedCallstack, DecodedCallstack );
+
+		// Convert back to FNames
+		for (const auto& It : DecodedCallstack)
+		{
+			NAME_INDEX NameIndex = 0;
+			TTypeFromString<NAME_INDEX>::FromString( NameIndex, *It );
+			const FName LongName = FName( NameIndex, NameIndex, 0 );
+
+			out_DecodedCallstack.Add( LongName );
+		}
+	}
+
+	/** Converts the encoded callstack into human readable callstack. */
+	static FString GetHumanReadable( const FName& EncodedCallstack )
+	{
+		TArray<FName> DecodedCallstack;
+		DecodeToNames( EncodedCallstack, DecodedCallstack );
+		const FString Result = GetHumanReadable( DecodedCallstack );
+		return Result;
+	}
+
+	/** Converts the encoded callstack into human readable callstack. */
+	static FString GetHumanReadable( const TArray<FName>& DecodedCallstack )
+	{
+		FString Result;
+
+		const int32 NumEntries = DecodedCallstack.Num();
+		//for (int32 Index = DecodedCallstack.Num() - 1; Index >= 0; --Index)
+		for (int32 Index = 0; Index < NumEntries; ++Index)
+		{
+			const FName LongName = DecodedCallstack[Index];
+			const FString ShortName = FStatNameAndInfo::GetShortNameFrom( LongName ).ToString();
+			//const FString Group = FStatNameAndInfo::GetGroupNameFrom( LongName ).ToString();
+			FString Desc = FStatNameAndInfo::GetDescriptionFrom( LongName );
+			Desc.Trim();
+
+			if (Desc.Len() == 0)
+			{
+				Result += ShortName;
+			}
+			else
+			{
+				Result += Desc;
+			}
+
+			if (Index != NumEntries - 1)
+			{
+				Result += TEXT( " -> " );
+			}
+		}
+
+		Result.ReplaceInline( TEXT( "STAT_" ), TEXT( "" ), ESearchCase::CaseSensitive );
+		return Result;
+	}
+
+protected:
+	/** Decodes encoded callstack to an array of strings. Where each string is the index of the FName. */
+	static void DecodeToStrings( const FName& EncodedCallstack, TArray<FString>& out_DecodedCallstack )
+	{
+		EncodedCallstack.ToString().ParseIntoArray( out_DecodedCallstack, CallstackSeparator, true );
+	}
+};
+
+const TCHAR* FStatsCallstack::CallstackSeparator = TEXT( "+" );
+
+/*-----------------------------------------------------------------------------
+	Allocation info
+-----------------------------------------------------------------------------*/
+
+FAllocationInfo::FAllocationInfo( uint64 InOldPtr, uint64 InPtr, int64 InSize, const TArray<FName>& InCallstack, uint32 InSequenceTag, EMemoryOperation InOp, bool bInHasBrokenCallstack ) 
+	: OldPtr( InOldPtr )
+	, Ptr( InPtr )
+	, Size( InSize )
+	, EncodedCallstack( *FStatsCallstack::Encode( InCallstack ) )
+	, SequenceTag( InSequenceTag )
+	, Op( InOp )
+	, bHasBrokenCallstack( bInHasBrokenCallstack )
+{
+
+}
+
+FAllocationInfo::FAllocationInfo( const FAllocationInfo& Other ) 
+	: OldPtr( Other.OldPtr )
+	, Ptr( Other.Ptr )
+	, Size( Other.Size )
+	, EncodedCallstack( Other.EncodedCallstack )
+	, SequenceTag( Other.SequenceTag )
+	, Op( Other.Op )
+	, bHasBrokenCallstack( Other.bHasBrokenCallstack )
+{
+
+}
+
+/*-----------------------------------------------------------------------------
+	FNodeAllocationInfo
+-----------------------------------------------------------------------------*/
+
+void FNodeAllocationInfo::SortBySize()
+{
+	ChildNodes.ValueSort( FNodeAllocationInfoSizeGreater() );
+	for (auto& It : ChildNodes)
+	{
+		It.Value->SortBySize();
+	}
+}
+
+
+void FNodeAllocationInfo::PrepareCallstackData( const TArray<FName>& InDecodedCallstack )
+{
+	DecodedCallstack = InDecodedCallstack;
+	EncodedCallstack = *FStatsCallstack::Encode( DecodedCallstack );
+	HumanReadableCallstack = FStatsCallstack::GetHumanReadable( DecodedCallstack );
+}
+
+/*-----------------------------------------------------------------------------
+	Stats stack helpers
+-----------------------------------------------------------------------------*/
+
+/** Holds stats stack state, used to preserve continuity when the game frame has changed. */
+struct FStackState
+{
+	/** Default constructor. */
+	FStackState()
+		: bIsBrokenCallstack( false )
+	{}
+
+	/** Call stack. */
+	TArray<FName> Stack;
+
+	/** Current function name. */
+	FName Current;
+
+	/** Whether this callstack is marked as broken due to mismatched start and end scope cycles. */
+	bool bIsBrokenCallstack;
+};
\ No newline at end of file
diff --git a/Engine/Source/Developer/Profiler/Private/ProfilerRawStats.h b/Engine/Source/Developer/Profiler/Private/ProfilerRawStats.h
new file mode 100644
--- /dev/null
+++ b/Engine/Source/Developer/Profiler/Private/ProfilerRawStats.h
@@ -0,0 +1,207 @@
+// Copyright 1998-2015 Epic Games, Inc. All Rights Reserved.
+
+#pragma once
+
+#include "StatsData.h"
+#include "StatsFile.h"
+
+/*-----------------------------------------------------------------------------
+	Basic structures
+-----------------------------------------------------------------------------*/
+
+/** Simple allocation info. Assumes the sequence tag are unique and doesn't wrap. */
+struct FAllocationInfo
+{
+	/** Old pointer, for Realloc. */
+	uint64 OldPtr;
+
+	/** Pointer, for Alloc, Realloc and Free. */
+	uint64 Ptr;
+
+	/** Size of the allocation. */
+	int64 Size;
+	
+	/**
+	 * Encoded callstack as FNames indices separated with +
+	 * 123+45+56+76
+	 */
+	FName EncodedCallstack;
+	
+	/**
+	 * Index of the memory operation. 
+	 * Stats are handled with packets sent from different threads.
+	 * This is used to sort the memory operations to be in the sequence order.
+	 */
+	uint32 SequenceTag;
+	
+	/** Memory operation, can be Alloc, Free or Realloc. */
+	EMemoryOperation Op;
+	
+	/** If true, the callstack is broken. Missing Start/End scope. */
+	bool bHasBrokenCallstack;
+
+	/** Initialization constructor. */
+	FAllocationInfo(
+		uint64 InOldPtr,
+		uint64 InPtr,
+		int64 InSize,
+		const TArray<FName>& InCallstack,
+		uint32 InSequenceTag,
+		EMemoryOperation InOp,
+		bool bInHasBrokenCallstack
+		);
+
+	/** Copy constructor. */
+	FAllocationInfo( const FAllocationInfo& Other );
+};
+
+/** Contains data about combined allocations for an unique scope. */
+struct FCombinedAllocationInfo
+{
+	/** Size as MB. */
+	double SizeMB;
+
+	/** Allocations count. */
+	int64 Count;
+
+	/** Human readable callstack like 'GameThread -> Preinit -> LoadModule */
+	FString HumanReadableCallstack;
+
+	/** Encoded callstack like '45+656+6565'. */
+	FName EncodedCallstack;
+
+	/** Callstack indices, encoded callstack parsed into an array [45,656,6565] */
+	TArray<FName> DecodedCallstack;
+
+	/** Size for this node and children. */
+	int64 Size;
+
+	/** Default constructor. */
+	FCombinedAllocationInfo()
+		: SizeMB( 0 )
+		, Count( 0 )
+		, Size( 0 )
+	{}
+
+	/** Operator for adding another combined allocation. */
+	FCombinedAllocationInfo& operator+=(const FCombinedAllocationInfo& Other)
+	{
+		if ((UPTRINT*)this != (UPTRINT*)&Other)
+		{
+			AddAllocation( Other.Size, Other.Count );
+		}
+		return *this;
+	}
+
+	/** Operator for subtracting another combined allocation. */
+	FCombinedAllocationInfo& operator-=(const FCombinedAllocationInfo& Other)
+	{
+		if ((UPTRINT*)this != (UPTRINT*)&Other)
+		{
+			AddAllocation( -Other.Size, -Other.Count );
+		}
+		return *this;
+	}
+
+	/** Operator for adding a new allocation. */
+	FCombinedAllocationInfo& operator+=(const FAllocationInfo& Other)
+	{
+		if ((UPTRINT*)this != (UPTRINT*)&Other)
+		{
+			AddAllocation( Other.Size, 1 );
+		}
+		return *this;
+	}
+	
+	/**
+	 * @return true, if the combined allocation is alive, the allocated data is larger than 0
+	 */
+	bool IsAlive()
+	{
+		return Size > 0;
+	}
+
+protected:
+	/** Adds a new allocation. */
+	void AddAllocation( int64 SizeToAdd, int64 CountToAdd = 1 )
+	{
+		Size += SizeToAdd;
+		Count += CountToAdd;
+
+		const double InvMB = 1.0 / 1024.0 / 1024.0;
+		SizeMB = Size * InvMB;
+	}
+};
+
+/** Contains data about node allocations, parent and children allocations. */
+struct FNodeAllocationInfo
+{
+	/** Size as MBs for this node and children. */
+	double SizeMB;
+
+	/** Allocations count for this node and children. */
+	int64 Count;
+
+	/** Human readable callstack like 'GameThread -> Preinit -> LoadModule */
+	FString HumanReadableCallstack;
+
+	/** Encoded callstack like '45+656+6565'. */
+	FName EncodedCallstack;
+
+	/** Callstack indices, encoded callstack parsed into an array [45,656,6565] */
+	TArray<FName> DecodedCallstack;
+
+	/** Parent node. */
+	FNodeAllocationInfo* Parent;
+
+	/** Child nodes. */
+	TMap<FName, FNodeAllocationInfo*> ChildNodes;
+
+	/** Size for this node and children. */
+	int64 Size;
+
+	/** Node's depth. */
+	int32 Depth;
+
+	/** Default constructor. */
+	FNodeAllocationInfo()
+		: SizeMB( 0.0 )
+		, Count( 0 )
+		, Parent( nullptr )
+		, Size( 0 )
+		, Depth( 0 )
+	{}
+
+	/** Destructor. */
+	~FNodeAllocationInfo()
+	{
+		DeleteAllChildrenNodes();
+	}
+
+	/** Accumulates this node with the other combined allocation. */
+	void Accumulate( const FCombinedAllocationInfo& Other )
+	{
+		Size += Other.Size;
+		Count += Other.Count;
+
+		const double InvMB = 1.0 / 1024.0 / 1024.0;
+		SizeMB = Size * InvMB;
+	}
+
+	/** Recursively sorts all children by size. */
+	void SortBySize();
+
+	/** Prepares all callstack data for this node. */
+	void PrepareCallstackData( const TArray<FName>& InDecodedCallstack );
+
+protected:
+	/** Recursively deletes all children. */
+	void DeleteAllChildrenNodes()
+	{
+		for (auto& It : ChildNodes)
+		{
+			delete It.Value;
+		}
+		ChildNodes.Empty();
+	}
+};
\ No newline at end of file
