Added missing file

diff --git a/Engine/Source/Developer/Profiler/Private/ProfilerRawStats.h b/Engine/Source/Developer/Profiler/Private/ProfilerRawStats.h
--- a/Engine/Source/Developer/Profiler/Private/ProfilerRawStats.h
+++ b/Engine/Source/Developer/Profiler/Private/ProfilerRawStats.h
@@ -204,4 +204,15 @@ protected:
 		}
 		ChildNodes.Empty();
 	}
-};
\ No newline at end of file
+};
+
+/*-----------------------------------------------------------------------------
+	FRawStatsMemoryProfiler
+-----------------------------------------------------------------------------*/
+
+/** Class for managing the raw stats in the context of memory profiling. */
+class FRawStatsMemoryProfiler
+{
+
+};
+
