Stats - Better support for RequestStop in ProcessStats, added stats for PostProcess stage, cleaned a bit code, some documentation added, cleaned/documented StatsMemoryDumpCommand, OpenRawStatsForMemoryProfiling is ready for use by the external code, ready for use

diff --git a/Engine/Source/Developer/Profiler/Private/ProfilerModule.cpp b/Engine/Source/Developer/Profiler/Private/ProfilerModule.cpp
--- a/Engine/Source/Developer/Profiler/Private/ProfilerModule.cpp
+++ b/Engine/Source/Developer/Profiler/Private/ProfilerModule.cpp
@@ -69,7 +69,6 @@ void FProfilerModule::Shutdown( TSharedRef<SDockTab> TabBeingClosed )
 
 void FProfilerModule::StatsMemoryDumpCommand( const TCHAR* Filename )
 {
-	//FRawStatsMemoryProfiler* Instance = FCreateStatsReader<FRawStatsMemoryProfiler>::ForRawStats( Filename );
 	TUniquePtr<FRawStatsMemoryProfiler> Instance( FCreateStatsReader<FRawStatsMemoryProfiler>::ForRawStats( Filename ) );
 	if (Instance)
 	{
@@ -81,46 +80,57 @@ void FProfilerModule::StatsMemoryDumpCommand( const TCHAR* Filename )
 
 			UE_LOG( LogStats, Log, TEXT( "Async: Stage: %s / %3i%%" ), *Instance->GetProcessingStageAsString(), Instance->GetStageProgress() );
 		}
-		//NewReader->RequestStop();
-
-		//Instance->
-
-		// Frame-240 Frame-120 Frame-060
-		TMap<FString, FCombinedAllocationInfo> FrameBegin_Exit;
-		Instance->CompareSnapshots_FString( TEXT( "BeginSnapshot" ), TEXT( "EngineLoop.Exit" ), FrameBegin_Exit );
-		Instance->DumpScopedAllocations( TEXT( "Begin_Exit" ), FrameBegin_Exit );
-
-#if	UE_BUILD_DEBUG
-		TMap<FString, FCombinedAllocationInfo> Frame060_120;
-		Instance->CompareSnapshots_FString( TEXT( "Frame-060" ), TEXT( "Frame-120" ), Frame060_120 );
-		Instance->DumpScopedAllocations( TEXT( "Frame060_120" ), Frame060_120 );
-
-		TMap<FString, FCombinedAllocationInfo> Frame060_240;
-		Instance->CompareSnapshots_FString( TEXT( "Frame-060" ), TEXT( "Frame-240" ), Frame060_240 );
-		Instance->DumpScopedAllocations( TEXT( "Frame060_240" ), Frame060_240 );
-
-		// Generate scoped tree view.
-		{
-			TMap<FName, FCombinedAllocationInfo> FrameBegin_Exit_FName;
-			Instance->CompareSnapshots_FName( TEXT( "BeginSnapshot" ), TEXT( "EngineLoop.Exit" ), FrameBegin_Exit_FName );
-
-			FNodeAllocationInfo Root;
-			Root.EncodedCallstack = TEXT( "ThreadRoot" );
-			Root.HumanReadableCallstack = TEXT( "ThreadRoot" );
-			Instance->GenerateScopedTreeAllocations( FrameBegin_Exit_FName, Root );
-		}
-
+		//Instance->RequestStop();
 
+		if (Instance->HasValidData())
 		{
-			TMap<FName, FCombinedAllocationInfo> Frame060_240_FName;
-			Instance->CompareSnapshots_FName( TEXT( "Frame-060" ), TEXT( "Frame-240" ), Frame060_240_FName );
-
-			FNodeAllocationInfo Root;
-			Root.EncodedCallstack = TEXT( "ThreadRoot" );
-			Root.HumanReadableCallstack = TEXT( "ThreadRoot" );
-			Instance->GenerateScopedTreeAllocations( Frame060_240_FName, Root );
-		}
+			// Dump scoped allocations between the first and the last snapshot.
+			const FName FirstSnapshotName = Instance->GetSnapshotNames()[0];
+			const FName LastSnapshotName = Instance->GetSnapshotNames()[Instance->GetSnapshotNames().Num()-1];
+
+			TMap<FString, FCombinedAllocationInfo> FrameBegin_End;
+			Instance->CompareSnapshotsHumanReadable( FirstSnapshotName, LastSnapshotName, FrameBegin_End );
+			Instance->DumpScopedAllocations( TEXT( "Begin_End" ), FrameBegin_End );
+
+			Instance->ProcessAndDumpUObjectAllocations( TEXT( "Frame-240" ) );
+
+#if	1/*UE_BUILD_DEBUG*/
+			// Dump debug scoped allocation generated when debug.EnableLeakTest=1 
+			TMap<FString, FCombinedAllocationInfo> Frame060_120;
+			Instance->CompareSnapshotsHumanReadable( TEXT( "Frame-060" ), TEXT( "Frame-120" ), Frame060_120 );
+			Instance->DumpScopedAllocations( TEXT( "Frame060_120" ), Frame060_120 );
+
+			TMap<FString, FCombinedAllocationInfo> Frame060_240;
+			Instance->CompareSnapshotsHumanReadable( TEXT( "Frame-060" ), TEXT( "Frame-240" ), Frame060_240 );
+			Instance->DumpScopedAllocations( TEXT( "Frame060_240" ), Frame060_240 );
+
+			// Generate scoped tree view.
+			{
+				// 1. Compare snapshots.
+				TMap<FName, FCombinedAllocationInfo> FrameBegin_End_FName;
+				Instance->CompareSnapshots( FirstSnapshotName, LastSnapshotName, FrameBegin_End_FName );
+
+				// 2. Generate tree of allocations.
+				FNodeAllocationInfo Root;
+				Root.EncodedCallstack = TEXT( "ThreadRoot" );
+				Root.HumanReadableCallstack = TEXT( "ThreadRoot" );
+				Instance->GenerateScopedTreeAllocations( FrameBegin_End_FName, Root );
+
+				// 3. Display.
+			}
+
+
+			{
+				TMap<FName, FCombinedAllocationInfo> Frame060_240_FName;
+				Instance->CompareSnapshots( TEXT( "Frame-060" ), TEXT( "Frame-240" ), Frame060_240_FName );
+
+				FNodeAllocationInfo Root;
+				Root.EncodedCallstack = TEXT( "ThreadRoot" );
+				Root.HumanReadableCallstack = TEXT( "ThreadRoot" );
+				Instance->GenerateScopedTreeAllocations( Frame060_240_FName, Root );
+			}
 #endif // UE_BUILD_DEBUG
+		}
 	}
 }
 
@@ -130,8 +140,11 @@ void FProfilerModule::StatsMemoryDumpCommand( const TCHAR* Filename )
 
 FRawStatsMemoryProfiler* FProfilerModule::OpenRawStatsForMemoryProfiling( const TCHAR* Filename )
 {
-	FRawStatsMemoryProfiler* Instance = 0;// new FRawStatsMemoryProfiler;
-
+	FRawStatsMemoryProfiler* Instance = FCreateStatsReader<FRawStatsMemoryProfiler>::ForRawStats( Filename );
+	if (Instance)
+	{
+		Instance->ReadAndProcessAsynchronously();
+	}
 	return Instance;
 }
 
diff --git a/Engine/Source/Developer/Profiler/Private/ProfilerRawStats.cpp b/Engine/Source/Developer/Profiler/Private/ProfilerRawStats.cpp
--- a/Engine/Source/Developer/Profiler/Private/ProfilerRawStats.cpp
+++ b/Engine/Source/Developer/Profiler/Private/ProfilerRawStats.cpp
@@ -191,6 +191,8 @@ void FNodeAllocationInfo::PrepareCallstackData( const TArray<FName>& InDecodedCa
 
 FRawStatsMemoryProfiler::FRawStatsMemoryProfiler( const TCHAR* InFilename )
 	: FStatsReadFile( InFilename, true )
+	, NumDuplicatedMemoryOperations( 0 )
+	, NumFWAMemoryOperations( 0 )
 	, NumMemoryOperations( 0 )
 	, LastSequenceTagForNamedMarker( 0 )
 {}
@@ -208,80 +210,107 @@ void FRawStatsMemoryProfiler::PostProcessStats()
 {
 	Super::PostProcessStats();
 
-	if (!SortSequenceAllocations())
+	const double StartTime = FPlatformTime::Seconds();
+
+	if (!IsProcessingStopped())
 	{
-		return;
+		SortSequenceAllocations();
+
+		// End marker.
+		Snapshots.Add( TPairInitializer<uint32, FName>( TNumericLimits<uint32>::Max(), TEXT( "EndSnapshot" ) ) );
+
+		// Copy snapshots.
+		SnapshotsToBeProcessed = Snapshots;
+
+		UE_LOG( LogStats, Log, TEXT( "NumMemoryOperations:   %i" ), NumMemoryOperations );
+		UE_LOG( LogStats, Log, TEXT( "SequenceAllocationNum: %i" ), SequenceAllocationArray.Num() );
+
+		GenerateAllocationMap();
+		DumpDebugAllocations();
 	}
 
-	// End marker.
-	Snapshots.Add( TPairInitializer<uint32, FName>( TNumericLimits<uint32>::Max(), TEXT( "EndSnapshot" ) ) );
+	// We don't need the allocation map. Each snapshot has its own copy.
+	AllocationMap.Empty();
+
+	if (!IsProcessingStopped())
+	{
+		StageProgress.Set( 100 );
+
+		const double TotalTime = FPlatformTime::Seconds() - StartTime;
+		UE_LOG( LogStats, Log, TEXT( "Post-Processing took %.2f sec(s)" ), TotalTime );
+	}
+	else
+	{
+		UE_LOG( LogStats, Warning, TEXT( "Post-Processing stopped, abandoning" ) );
+	}
+}
 
-	// Copy snapshots.
-	SnapshotsToBeProcessed = Snapshots;
+void FRawStatsMemoryProfiler::DumpDebugAllocations()
+{
+#if	UE_BUILD_DEBUG
 
-	UE_LOG( LogStats, Warning, TEXT( "NumMemoryOperations:   %llu" ), NumMemoryOperations );
-	UE_LOG( LogStats, Warning, TEXT( "SequenceAllocationNum: %i" ), SequenceAllocationArray.Num() );
+	// Dump problematic allocations
+	DuplicatedAllocMap.ValueSort( FAllocationInfoSizeGreater() );
 
-	// Pass 2.
-	/*
-	TMap<uint32, FAllocationInfo> UniqueSeq;
-	TMultiMap<uint32, FAllocationInfo> OriginalAllocs;
-	TMultiMap<uint32, FAllocationInfo> BrokenAllocs;
-	for (const FAllocationInfo& Alloc : SequenceAllocationArray)
+	uint64 TotalDuplicatedMemory = 0;
+	for (const auto& It : DuplicatedAllocMap)
 	{
-		const FAllocationInfo* Found = UniqueSeq.Find( Alloc.SequenceTag );
-		if (!Found)
-		{
-			UniqueSeq.Add( Alloc.SequenceTag, Alloc );
-		}
-		else
+		const FAllocationInfo& Alloc = It.Value;
+		TotalDuplicatedMemory += Alloc.Size;
+	}
+
+	UE_LOG( LogStats, Warning, TEXT( "Dumping duplicated alloc map" ) );
+	const float MaxPctDisplayed = 0.80f;
+	uint64 DisplayedSoFar = 0;
+	for (const auto& It : DuplicatedAllocMap)
+	{
+		const FAllocationInfo& Alloc = It.Value;
+		const FString AllocCallstack = FStatsCallstack::GetHumanReadable( Alloc.EncodedCallstack );
+		UE_LOG( LogStats, Log, TEXT( "%lli (%.2f MB) %s" ), Alloc.Size, Alloc.Size / 1024.0f / 1024.0f, *AllocCallstack );
+
+		DisplayedSoFar += Alloc.Size;
+
+		const float CurrentPct = (float)DisplayedSoFar / (float)TotalDuplicatedMemory;
+		if (CurrentPct > MaxPctDisplayed)
 		{
-			OriginalAllocs.Add( Alloc.SequenceTag, *Found );
-			BrokenAllocs.Add( Alloc.SequenceTag, Alloc );
+			break;
 		}
 	}
-	//*/
 
-	// Alive allocations.
-	TMap<uint64, FAllocationInfo> AllocationMap;
-	TMultiMap<uint64, FAllocationInfo> FreeWithoutAllocMap;
-	TMultiMap<uint64, FAllocationInfo> DuplicatedAllocMap;
-	int32 NumDuplicatedMemoryOperations = 0;
-	int32 NumFWAMemoryOperations = 0; // FreeWithoutAlloc
-	int32 NumZeroAllocs = 0; // Malloc(0)
+#endif // UE_BUILD_DEBUG
+}
+
+void FRawStatsMemoryProfiler::FreeDebugInformation()
+{
+	FreeWithoutAllocMap.Empty();
+	DuplicatedAllocMap.Empty();
+}
 
+void FRawStatsMemoryProfiler::GenerateAllocationMap()
+{
 	// Initialize the begin snapshot.
 	auto BeginSnapshot = SnapshotsToBeProcessed[0];
+
 	SnapshotsToBeProcessed.RemoveAt( 0 );
 	PrepareSnapshot( BeginSnapshot.Value, AllocationMap );
-
 	auto CurrentSnapshot = SnapshotsToBeProcessed[0];
 
-	UE_LOG( LogStats, Warning, TEXT( "Generating memory operations map" ) );
+	UE_LOG( LogStats, Log, TEXT( "Generating memory operations map" ) );
+
 	const int32 NumSequenceAllocations = SequenceAllocationArray.Num();
 	const int32 OnePercent = FMath::Max( NumSequenceAllocations / 100, 1024 );
-	for (int32 Index = 0; Index < NumSequenceAllocations; Index++)
+	for (int32 AllocationIndex = 0; AllocationIndex < NumSequenceAllocations; AllocationIndex++)
 	{
-		if (Index % OnePercent)
+		if (AllocationIndex % OnePercent == 0)
 		{
-			const double CurrentSeconds = FPlatformTime::Seconds();
-			if (CurrentSeconds > LastUpdateTime + NumSecondsBetweenUpdates)
+			UpdateGenerateMemoryMapProgress( AllocationIndex );
+			if (IsProcessingStopped())
 			{
-				const int32 PercentagePos = int32( 100.0*(Index + 1) / NumSequenceAllocations );
-				StageProgress.Set( PercentagePos );
-				UE_LOG( LogStats, Verbose, TEXT( "Processing allocations %3i%% (%10i/%10i)" ), PercentagePos, Index + 1, NumSequenceAllocations );
-				LastUpdateTime = CurrentSeconds;
-
-				// Abandon support, simply return for now.
-				if (bShouldStopProcessing == true)
-				{
-					SetProcessingStage( EStatsProcessingStage::SPS_Stopped );
-					return;
-				}
+				break;
 			}
 		}
 
-		const FAllocationInfo& Alloc = SequenceAllocationArray[Index];
+		const FAllocationInfo& Alloc = SequenceAllocationArray[AllocationIndex];
 
 		// Check named marker/snapshots
 		if (Alloc.SequenceTag > CurrentSnapshot.Key)
@@ -306,10 +335,10 @@ void FRawStatsMemoryProfiler::PostProcessStats()
 			}
 			else
 			{
+#if	UE_BUILD_DEBUG
 				const FAllocationInfo* FoundAndFreed = FreeWithoutAllocMap.Find( Found->Ptr );
 				const FAllocationInfo* FoundAndAllocated = FreeWithoutAllocMap.Find( Alloc.Ptr );
 
-#if	UE_BUILD_DEBUG
 				if (FoundAndFreed)
 				{
 					const FString FoundAndFreedCallstack = FStatsCallstack::GetHumanReadable( FoundAndFreed->EncodedCallstack );
@@ -365,7 +394,7 @@ void FRawStatsMemoryProfiler::PostProcessStats()
 
 #if	UE_BUILD_DEBUG
 				const FString ReallocCallstack = FStatsCallstack::GetHumanReadable( Alloc.EncodedCallstack );
-				//UE_LOG( LogStats, Warning, TEXT( "ReallocWithoutAlloc: %s %i %i/%i [%i]" ), *ReallocCallstack, Alloc.Size, Alloc.OldPtr, Alloc.Ptr, Alloc.SequenceTag );
+				UE_LOG( LogStats, VeryVerbose, TEXT( "ReallocWithoutAlloc: %s %i %i/%i [%i]" ), *ReallocCallstack, Alloc.Size, Alloc.OldPtr, Alloc.Ptr, Alloc.SequenceTag );
 #endif // UE_BUILD_DEBUG
 
 				AllocationMap.Add( Alloc.Ptr, Alloc );
@@ -391,48 +420,38 @@ void FRawStatsMemoryProfiler::PostProcessStats()
 
 #if	UE_BUILD_DEBUG
 				const FString FWACallstack = FStatsCallstack::GetHumanReadable( Alloc.EncodedCallstack );
-				//UE_LOG( LogStats, Warning, TEXT( "FreeWithoutAllocCallstack: %s %i" ), *FWACallstack, Alloc.Ptr );
+				UE_LOG( LogStats, VeryVerbose, TEXT( "FreeWithoutAllocCallstack: %s %i" ), *FWACallstack, Alloc.Ptr );
 #endif // UE_BUILD_DEBUG
 			}
 		}
 	}
-	StageProgress.Set( 100 );
 
 	auto EndSnapshot = SnapshotsToBeProcessed[0];
 	SnapshotsToBeProcessed.RemoveAt( 0 );
 	PrepareSnapshot( EndSnapshot.Value, AllocationMap );
 
-	UE_LOG( LogStats, Warning, TEXT( "NumDuplicatedMemoryOperations: %i" ), NumDuplicatedMemoryOperations );
-	UE_LOG( LogStats, Warning, TEXT( "NumFWAMemoryOperations:        %i" ), NumFWAMemoryOperations );
-	UE_LOG( LogStats, Warning, TEXT( "NumZeroAllocs:                 %i" ), NumZeroAllocs );
+	SnapshotNamesArray = SnapshotNamesSet.Array();
 
-	// Dump problematic allocations
-	DuplicatedAllocMap.ValueSort( FAllocationInfoSizeGreater() );
-	//FreeWithoutAllocMap
+	UE_LOG( LogStats, Verbose, TEXT( "NumDuplicatedMemoryOperations: %i" ), NumDuplicatedMemoryOperations );
+	UE_LOG( LogStats, Verbose, TEXT( "NumFWAMemoryOperations:        %i" ), NumFWAMemoryOperations );
+	UE_LOG( LogStats, Verbose, TEXT( "NumZeroAllocs:                 %i" ), NumZeroAllocs );
+}
 
-	uint64 TotalDuplicatedMemory = 0;
-	for (const auto& It : DuplicatedAllocMap)
+void FRawStatsMemoryProfiler::UpdateGenerateMemoryMapProgress( const int32 AllocationIndex )
+{
+	const double CurrentSeconds = FPlatformTime::Seconds();
+	if (CurrentSeconds > LastUpdateTime + NumSecondsBetweenUpdates)
 	{
-		const FAllocationInfo& Alloc = It.Value;
-		TotalDuplicatedMemory += Alloc.Size;
+		const int32 PercentagePos = int32( 100.0*AllocationIndex / SequenceAllocationArray.Num() );
+		StageProgress.Set( PercentagePos );
+		UE_LOG( LogStats, Verbose, TEXT( "Processing allocations %3i%% (%10i/%10i)" ), PercentagePos, AllocationIndex, SequenceAllocationArray.Num() );
+		LastUpdateTime = CurrentSeconds;
 	}
-
-	UE_LOG( LogStats, Warning, TEXT( "Dumping duplicated alloc map" ) );
-	const float MaxPctDisplayed = 0.80f;
-	uint64 DisplayedSoFar = 0;
-	for (const auto& It : DuplicatedAllocMap)
+	
+	// Abandon support.
+	if (bShouldStopProcessing == true)
 	{
-		const FAllocationInfo& Alloc = It.Value;
-		const FString AllocCallstack = FStatsCallstack::GetHumanReadable( Alloc.EncodedCallstack );
-		UE_LOG( LogStats, Log, TEXT( "%lli (%.2f MB) %s" ), Alloc.Size, Alloc.Size / 1024.0f / 1024.0f, *AllocCallstack );
-
-		DisplayedSoFar += Alloc.Size;
-
-		const float CurrentPct = (float)DisplayedSoFar / (float)TotalDuplicatedMemory;
-		if (CurrentPct > MaxPctDisplayed)
-		{
-			break;
-		}
+		SetProcessingStage( EStatsProcessingStage::SPS_Stopped );
 	}
 }
 
@@ -502,8 +521,10 @@ void FRawStatsMemoryProfiler::ProcessMemoryOperation( EMemoryOperation MemOp, ui
 	}
 }
 
-bool FRawStatsMemoryProfiler::SortSequenceAllocations()
+void FRawStatsMemoryProfiler::SortSequenceAllocations()
 {
+	FScopeLogTime SLT( TEXT( "SortSequenceAllocations" ), nullptr, FScopeLogTime::ScopeLog_Milliseconds );
+
 	// Sort all memory operation by the sequence tag, iterate through all operation and generate memory usage.
 	SequenceAllocationArray.Sort( FAllocationInfoSequenceTagLess() );
 
@@ -511,20 +532,14 @@ bool FRawStatsMemoryProfiler::SortSequenceAllocations()
 	if (bShouldStopProcessing == true)
 	{
 		SetProcessingStage( EStatsProcessingStage::SPS_Stopped );
-		return false;
-	}
-	else
-	{
-		return true;
 	}
 }
 
 void FRawStatsMemoryProfiler::GenerateScopedTreeAllocations( const TMap<FName, FCombinedAllocationInfo>& ScopedAllocations, FNodeAllocationInfo& out_Root )
 {
-	// Experimental code, partially optimized.
-	// Decode all scoped allocations.
-	// Generate tree for allocations and combine them.
+	FScopeLogTime SLT( TEXT( "GenerateScopedTreeAllocations" ), nullptr, FScopeLogTime::ScopeLog_Milliseconds );
 
+	// Decode all scoped allocations, generate tree for allocations and combine them.
 	for (const auto& It : ScopedAllocations)
 	{
 		const FName& EncodedCallstack = It.Key;
@@ -576,95 +591,16 @@ void FRawStatsMemoryProfiler::GenerateScopedTreeAllocations( const TMap<FName, F
 	out_Root.SortBySize();
 }
 
-
-void FRawStatsMemoryProfiler::GenerateMemoryUsageReport( const TMap<uint64, FAllocationInfo>& AllocationMap, const TSet<FName>& UObjectRawNames )
+void FRawStatsMemoryProfiler::ProcessAndDumpUObjectAllocations( const FName SnapshotName )
 {
-	if (AllocationMap.Num() == 0)
-	{
-		UE_LOG( LogStats, Warning, TEXT( "There are no allocations, make sure memory profiler is enabled" ) );
-	}
-	else
+	if (!SnapshotsWithAllocationMap.Contains(SnapshotName))
 	{
-		//ProcessAndDumpUObjectAllocations( AllocationMap );
-		ProcessAndDumpScopedAllocations( AllocationMap );
-	}
-}
-
-void FRawStatsMemoryProfiler::ProcessAndDumpScopedAllocations( const TMap<uint64, FAllocationInfo>& AllocationMap )
-{
-	// This code is not optimized. 
-	FScopeLogTime SLT( TEXT( "ProcessingScopedAllocations" ), nullptr, FScopeLogTime::ScopeLog_Seconds );
-	UE_LOG( LogStats, Warning, TEXT( "Processing scoped allocations" ) );
-
-	const FString ReportName = FString::Printf( TEXT( "%s-Memory-Scoped" ), *GetPlatformName() );
-	FDiagnosticTableViewer MemoryReport( *FDiagnosticTableViewer::GetUniqueTemporaryFilePath( *ReportName ), true );
-
-	// Write a row of headings for the table's columns.
-	MemoryReport.AddColumn( TEXT( "Size (bytes)" ) );
-	MemoryReport.AddColumn( TEXT( "Size (MB)" ) );
-	MemoryReport.AddColumn( TEXT( "Count" ) );
-	MemoryReport.AddColumn( TEXT( "Callstack" ) );
-	MemoryReport.CycleRow();
-
-	TMap<FName, FCombinedAllocationInfo> CombinedAllocations;
-	uint64 TotalAllocatedMemory = 0;
-	uint64 NumAllocations = 0;
-	GenerateScopedAllocations( AllocationMap, CombinedAllocations, TotalAllocatedMemory, NumAllocations );
-
-	// Dump memory to the log.
-	CombinedAllocations.ValueSort( FCombinedAllocationInfoSizeGreater() );
-
-	const float MaxPctDisplayed = 0.90f;
-	int32 CurrentIndex = 0;
-	uint64 DisplayedSoFar = 0;
-	UE_LOG( LogStats, Warning, TEXT( "Index, Size (Size MB), Count, Stat desc" ) );
-	for (const auto& It : CombinedAllocations)
-	{
-		const FCombinedAllocationInfo& CombinedAllocation = It.Value;
-		const FName& EncodedCallstack = It.Key;
-
-		const FString AllocCallstack = FStatsCallstack::GetHumanReadable( EncodedCallstack );
-
-		UE_LOG( LogStats, Log, TEXT( "%2i, %llu (%.2f MB), %llu, %s" ),
-				CurrentIndex,
-				CombinedAllocation.Size,
-				CombinedAllocation.Size / 1024.0f / 1024.0f,
-				CombinedAllocation.Count,
-				*AllocCallstack );
-
-		// Dump stats
-		MemoryReport.AddColumn( TEXT( "%llu" ), CombinedAllocation.Size );
-		MemoryReport.AddColumn( TEXT( "%.2f MB" ), CombinedAllocation.Size / 1024.0f / 1024.0f );
-		MemoryReport.AddColumn( TEXT( "%llu" ), CombinedAllocation.Count );
-		MemoryReport.AddColumn( *AllocCallstack );
-		MemoryReport.CycleRow();
-
-		CurrentIndex++;
-		DisplayedSoFar += CombinedAllocation.Size;
-
-		const float CurrentPct = (float)DisplayedSoFar / (float)TotalAllocatedMemory;
-		if (CurrentPct > MaxPctDisplayed)
-		{
-			break;
-		}
+		UE_LOG( LogStats, Warning, TEXT( "Snapshot not found: %s" ), *SnapshotName.ToString() );
+		return;
 	}
 
-	UE_LOG( LogStats, Warning, TEXT( "Allocated memory: %llu bytes (%.2f MB)" ), TotalAllocatedMemory, TotalAllocatedMemory / 1024.0f / 1024.0f );
+	const TMap<uint64, FAllocationInfo>& AllocationMap = SnapshotsWithAllocationMap.FindChecked( SnapshotName );
 
-	// Add a total row.
-	MemoryReport.CycleRow();
-	MemoryReport.CycleRow();
-	MemoryReport.CycleRow();
-	MemoryReport.AddColumn( TEXT( "%llu" ), TotalAllocatedMemory );
-	MemoryReport.AddColumn( TEXT( "%.2f MB" ), TotalAllocatedMemory / 1024.0f / 1024.0f );
-	MemoryReport.AddColumn( TEXT( "%llu" ), NumAllocations );
-	MemoryReport.AddColumn( TEXT( "TOTAL" ) );
-	MemoryReport.CycleRow();
-}
-
-void FRawStatsMemoryProfiler::ProcessAndDumpUObjectAllocations( const TMap<uint64, FAllocationInfo>& AllocationMap, const TSet<FName>& UObjectRawNames )
-{
-	// This code is not optimized. 
 	FScopeLogTime SLT( TEXT( "ProcessingUObjectAllocations" ), nullptr, FScopeLogTime::ScopeLog_Seconds );
 	UE_LOG( LogStats, Warning, TEXT( "Processing UObject allocations" ) );
 
@@ -725,13 +661,13 @@ void FRawStatsMemoryProfiler::ProcessAndDumpUObjectAllocations( const TMap<uint6
 	const float MaxPctDisplayed = 0.90f;
 	int32 CurrentIndex = 0;
 	uint64 DisplayedSoFar = 0;
-	UE_LOG( LogStats, Warning, TEXT( "Index, Size (Size MB), Count, UObject class" ) );
+	UE_LOG( LogStats, VeryVerbose, TEXT( "Index, Size (Size MB), Count, UObject class" ) );
 	for (const auto& It : UObjectAllocations)
 	{
 		const FCombinedAllocationInfo& CombinedAllocation = It.Value;
 		const FName& UObjectClass = It.Key;
 
-		UE_LOG( LogStats, Log, TEXT( "%2i, %llu (%.2f MB), %llu, %s" ),
+		UE_LOG( LogStats, VeryVerbose, TEXT( "%2i, %llu (%.2f MB), %llu, %s" ),
 				CurrentIndex,
 				CombinedAllocation.Size,
 				CombinedAllocation.Size / 1024.0f / 1024.0f,
@@ -755,7 +691,7 @@ void FRawStatsMemoryProfiler::ProcessAndDumpUObjectAllocations( const TMap<uint6
 		}
 	}
 
-	UE_LOG( LogStats, Warning, TEXT( "Allocated memory: %llu bytes (%.2f MB)" ), TotalAllocatedMemory, TotalAllocatedMemory / 1024.0f / 1024.0f );
+	UE_LOG( LogStats, VeryVerbose, TEXT( "Allocated memory: %llu bytes (%.2f MB)" ), TotalAllocatedMemory, TotalAllocatedMemory / 1024.0f / 1024.0f );
 
 	// Add a total row.
 	MemoryReport.CycleRow();
@@ -776,7 +712,6 @@ void FRawStatsMemoryProfiler::DumpScopedAllocations( const TCHAR* Name, const TM
 		return;
 	}
 
-	// This code is not optimized. 
 	FScopeLogTime SLT( TEXT( "ProcessingScopedAllocations" ), nullptr, FScopeLogTime::ScopeLog_Seconds );
 	UE_LOG( LogStats, Warning, TEXT( "Dumping scoped allocations: %s" ), Name );
 
@@ -790,19 +725,18 @@ void FRawStatsMemoryProfiler::DumpScopedAllocations( const TCHAR* Name, const TM
 	MemoryReport.AddColumn( TEXT( "Callstack" ) );
 	MemoryReport.CycleRow();
 
-
 	FCombinedAllocationInfo Total;
 
 	const float MaxPctDisplayed = 0.90f;
 	int32 CurrentIndex = 0;
-	UE_LOG( LogStats, Warning, TEXT( "Index, Size (Size MB), Count, Stat desc" ) );
+	UE_LOG( LogStats, VeryVerbose, TEXT( "Index, Size (Size MB), Count, Stat desc" ) );
 	for (const auto& It : CombinedAllocations)
 	{
 		const FCombinedAllocationInfo& CombinedAllocation = It.Value;
 		//const FName& EncodedCallstack = It.Key;
 		const FString AllocCallstack = It.Key;// GetCallstack( EncodedCallstack );
 
-		UE_LOG( LogStats, Log, TEXT( "%2i, %llu (%.2f MB), %llu, %s" ),
+		UE_LOG( LogStats, VeryVerbose, TEXT( "%2i, %llu (%.2f MB), %llu, %s" ),
 				CurrentIndex,
 				CombinedAllocation.Size,
 				CombinedAllocation.Size / 1024.0f / 1024.0f,
@@ -820,7 +754,7 @@ void FRawStatsMemoryProfiler::DumpScopedAllocations( const TCHAR* Name, const TM
 		Total += CombinedAllocation;
 	}
 
-	UE_LOG( LogStats, Warning, TEXT( "Allocated memory: %llu bytes (%.2f MB)" ), Total.Size, Total.SizeMB );
+	UE_LOG( LogStats, VeryVerbose, TEXT( "Allocated memory: %llu bytes (%.2f MB)" ), Total.Size, Total.SizeMB );
 
 	// Add a total row.
 	MemoryReport.CycleRow();
@@ -857,11 +791,11 @@ void FRawStatsMemoryProfiler::PrepareSnapshot( const FName SnapshotName, const T
 
 	// Make sure the snapshot name is unique.
 	FName UniqueSnapshotName = SnapshotName;
-	while (SnapshotNames.Contains( UniqueSnapshotName ))
+	while (SnapshotNamesSet.Contains( UniqueSnapshotName ))
 	{
 		UniqueSnapshotName = FName( UniqueSnapshotName, UniqueSnapshotName.GetNumber() + 1 );
 	}
-	SnapshotNames.Add( UniqueSnapshotName );
+	SnapshotNamesSet.Add( UniqueSnapshotName );
 
 	SnapshotsWithAllocationMap.Add( UniqueSnapshotName, AllocationMap );
 
@@ -884,8 +818,10 @@ void FRawStatsMemoryProfiler::PrepareSnapshot( const FName SnapshotName, const T
 	UE_LOG( LogStats, Warning, TEXT( "PrepareSnapshot: %s Alloc: %i Scoped: %i Total: %.2f MB" ), *UniqueSnapshotName.ToString(), AllocationMap.Num(), SnapshotCombinedAllocations.Num(), TotalAllocatedMemory / 1024.0f / 1024.0f );
 }
 
-void FRawStatsMemoryProfiler::CompareSnapshots_FName( const FName BeginSnaphotName, const FName EndSnaphotName, TMap<FName, FCombinedAllocationInfo>& out_Result )
+void FRawStatsMemoryProfiler::CompareSnapshots( const FName BeginSnaphotName, const FName EndSnaphotName, TMap<FName, FCombinedAllocationInfo>& out_Result )
 {
+	FScopeLogTime SLT( TEXT( "CompareSnapshots" ), nullptr, FScopeLogTime::ScopeLog_Milliseconds );
+
 	const auto BeginSnaphotPtr = SnapshotsWithScopedAllocations.Find( BeginSnaphotName );
 	const auto EndSnapshotPtr = SnapshotsWithScopedAllocations.Find( EndSnaphotName );
 	if (BeginSnaphotPtr && EndSnapshotPtr)
@@ -923,8 +859,10 @@ void FRawStatsMemoryProfiler::CompareSnapshots_FName( const FName BeginSnaphotNa
 	}
 }
 
-void FRawStatsMemoryProfiler::CompareSnapshots_FString( const FName BeginSnaphotName, const FName EndSnaphotName, TMap<FString, FCombinedAllocationInfo>& out_Result )
+void FRawStatsMemoryProfiler::CompareSnapshotsHumanReadable( const FName BeginSnaphotName, const FName EndSnaphotName, TMap<FString, FCombinedAllocationInfo>& out_Result )
 {
+	FScopeLogTime SLT( TEXT( "CompareSnapshotsHumanReadable" ), nullptr, FScopeLogTime::ScopeLog_Milliseconds );
+
 	const auto BeginSnaphotPtr = SnapshotsWithDecodedScopedAllocations.Find( BeginSnaphotName );
 	const auto EndSnapshotPtr = SnapshotsWithDecodedScopedAllocations.Find( EndSnaphotName );
 	if (BeginSnaphotPtr && EndSnapshotPtr)
diff --git a/Engine/Source/Developer/Profiler/Private/ProfilerRawStats.h b/Engine/Source/Developer/Profiler/Private/ProfilerRawStats.h
--- a/Engine/Source/Developer/Profiler/Private/ProfilerRawStats.h
+++ b/Engine/Source/Developer/Profiler/Private/ProfilerRawStats.h
@@ -265,22 +265,30 @@ protected:
 	/** Processes memory operation. @see EMemoryOperation. */
 	virtual void ProcessMemoryOperation( EMemoryOperation MemOp, uint64 Ptr, uint64 NewPtr, int64 Size, uint32 SequenceTag, const FStackState& StackState ) override;
 
+	/**
+	 * Updates process stage progress periodically, does debug logging if enabled.
+	 */
+	void UpdateGenerateMemoryMapProgress( const int32 AllocationIndex );
+
 	/**
 	 * Sorts sequence allocations.
-	 *
-	 * @return true, if finished, false otherwise like stopped.
 	 */
-	bool SortSequenceAllocations();
+	void SortSequenceAllocations();
 
-public:
-	/** Generates a basic memory usage report and prints it to the log. */
-	void GenerateMemoryUsageReport( const TMap<uint64, FAllocationInfo>& AllocationMap, const TSet<FName>& UObjectRawNames );
+	/**
+	 *	Generates allocation map and processes all memory snapshots .
+	 */
+	void GenerateAllocationMap();
 
-	/** Generates scoped allocation statistics. */
-	void ProcessAndDumpScopedAllocations( const TMap<uint64, FAllocationInfo>& AllocationMap );
+	/** Dumps debug allocations. */
+	void DumpDebugAllocations();
 
+	/** Frees memory used by the debug information. */
+	void FreeDebugInformation();
+
+public:
 	/** Generates UObject allocation statistics. */
-	void ProcessAndDumpUObjectAllocations( const TMap<uint64, FAllocationInfo>& AllocationMap, const TSet<FName>& UObjectRawNames );
+	void ProcessAndDumpUObjectAllocations( const FName SnapshotName );
 
 	void DumpScopedAllocations( const TCHAR* Name, const TMap<FString, FCombinedAllocationInfo>& CombinedAllocations );
 
@@ -292,34 +300,54 @@ public:
 	/** Prepares data for a snapshot. */
 	void PrepareSnapshot( const FName SnapshotName, const TMap<uint64, FAllocationInfo>& AllocationMap );
 
-	void CompareSnapshots_FName( const FName BeginSnaphotName, const FName EndSnaphotName, TMap<FName, FCombinedAllocationInfo>& out_Result );
+	/** Compare two snapshots and saves the result for further processing. */
+	void CompareSnapshots( const FName BeginSnaphotName, const FName EndSnaphotName, TMap<FName, FCombinedAllocationInfo>& out_Result );
 
-	void CompareSnapshots_FString( const FName BeginSnaphotName, const FName EndSnaphotName, TMap<FString, FCombinedAllocationInfo>& out_Result );
+	/** Compare two snapshots and saves the result in human readable format, results can be saved to a file. */
+	void CompareSnapshotsHumanReadable( const FName BeginSnaphotName, const FName EndSnaphotName, TMap<FString, FCombinedAllocationInfo>& out_Result );
 
+	/**
+	 * @return a set of all snapshot names.
+	 */
+	const TArray<FName> GetSnapshotNames() const
+	{
+		return SnapshotNamesArray;
+	}
 
 	/**
-	* @return Platform's name based on the loaded ue4statsraw file
-	*/
+	 * @return Platform's name based on the loaded ue4statsraw file.
+	 */
 	const FString& GetPlatformName() const
 	{
 		return Header.PlatformName;
 	}
 
+	/**
+	 * @return true, if the file contains any usable data.
+	 */
+	bool HasValidData() const
+	{
+		return SnapshotNamesArray.Num() >= 2;
+	}
+
 protected:
 	/**
-	 *	 All allocation ordered by the sequence tag.
+	 *	 All allocations ordered by the sequence tag.
 	 *	 There is an assumption that the sequence tag will not turn-around.
 	 */
 	TArray<FAllocationInfo> SequenceAllocationArray;
 
-	/** Allocation map. */
-	//TMap<uint64, FAllocationInfo>& AllocationMap;
+	/** Map of currently alive allocations. Ptr to AllocationInfo. */
+	TMap<uint64, FAllocationInfo> AllocationMap;
 
 	/** The sequence tag mapping to the named markers. */
 	TArray<TPair<uint32, FName>> Snapshots;
 
 	/** Unique snapshot names. */
-	TSet<FName> SnapshotNames;
+	TSet<FName> SnapshotNamesSet;
+
+	/** Unique snapshot names, the same as SnapshotNamesSet. */
+	TArray<FName> SnapshotNamesArray;
 
 	/** The sequence tag mapping to the named markers that need to processed. */
 	TArray<TPair<uint32, FName>> SnapshotsToBeProcessed;
@@ -332,10 +360,25 @@ protected:
 
 	TMap<FName, TMap<FString, FCombinedAllocationInfo> > SnapshotsWithDecodedScopedAllocations;
 
+	/** Free without allocation map. Debug only. */
+	TMultiMap<uint64, FAllocationInfo> FreeWithoutAllocMap;
+
+	/** Duplicated allocation map. Debug only. */
+	TMultiMap<uint64, FAllocationInfo> DuplicatedAllocMap;
+
+	/** Number of duplicated memory operations. */
+	int32 NumDuplicatedMemoryOperations;
+
+	/** Number of Free without Alloc. */
+	int32 NumFWAMemoryOperations;
+
+	/** Number of Malloc(0). */
+	int32 NumZeroAllocs = 0;
+
 	/** Number of memory operations. */
-	uint64 NumMemoryOperations;
+	int32 NumMemoryOperations;
 
 	/** Last sequence tag for named marker. */
-	uint32 LastSequenceTagForNamedMarker;
+	int32 LastSequenceTagForNamedMarker;
 };
 
diff --git a/Engine/Source/Developer/Profiler/Public/Interfaces/IProfilerModule.h b/Engine/Source/Developer/Profiler/Public/Interfaces/IProfilerModule.h
--- a/Engine/Source/Developer/Profiler/Public/Interfaces/IProfilerModule.h
+++ b/Engine/Source/Developer/Profiler/Public/Interfaces/IProfilerModule.h
@@ -24,6 +24,9 @@ public:
 	/** Implements stats memory dump command. */
 	virtual void StatsMemoryDumpCommand( const TCHAR* Filename ) = 0;
 
-	/** Creates a new instance of the memory profiler based the raw stats. */
+	/** 
+	 * Creates a new instance of the memory profiler based the raw stats. 
+	 * When no longer needed must be stopped via RequestStop() and deleted to avoid memory leaks.
+	 */
 	virtual FRawStatsMemoryProfiler* OpenRawStatsForMemoryProfiling( const TCHAR* Filename ) = 0;
 };
\ No newline at end of file
diff --git a/Engine/Source/Runtime/Core/Private/Stats/StatsFile.cpp b/Engine/Source/Runtime/Core/Private/Stats/StatsFile.cpp
--- a/Engine/Source/Runtime/Core/Private/Stats/StatsFile.cpp
+++ b/Engine/Source/Runtime/Core/Private/Stats/StatsFile.cpp
@@ -570,7 +570,7 @@ void FStatsReadFile::ReadAndProcessSynchronously()
 	ProcessStats();
 	PostProcessStats();
 
-	if (GetProcessingStage() == EStatsProcessingStage::SPS_Stopped)
+	if (IsProcessingStopped())
 	{
 		SetProcessingStage( EStatsProcessingStage::SPS_Invalid );
 	}
@@ -774,7 +774,8 @@ void FStatsReadFile::ReadStats()
 				FileInfo.MaximumPacketSize = FMath::Max<int32>( FileInfo.MaximumPacketSize, StatPacket->StatMessages.GetAllocatedSize() );
 			}
 
-			if (!UpdateReadStageProgress())
+			UpdateReadStageProgress();
+			if (IsProcessingStopped())
 			{
 				break;
 			}
@@ -788,7 +789,7 @@ void FStatsReadFile::ReadStats()
 		// Verify that frames are sequential.
 		check( Frames[Frames.Num() - 1] == Frames.Num() );
 
-		if (GetProcessingStage() != EStatsProcessingStage::SPS_Stopped)
+		if (!IsProcessingStopped())
 		{
 			const double TotalTime = FPlatformTime::Seconds() - StartTime;
 			UE_LOG( LogStats, Log, TEXT( "Reading took %.2f sec(s)" ), TotalTime );
@@ -806,7 +807,7 @@ void FStatsReadFile::ReadStats()
 
 void FStatsReadFile::PreProcessStats()
 {
-	if (GetProcessingStage() != EStatsProcessingStage::SPS_Stopped)
+	if (!IsProcessingStopped())
 	{
 		SetProcessingStage( EStatsProcessingStage::SPS_PreProcessStats );
 	}
@@ -816,7 +817,7 @@ void FStatsReadFile::ProcessStats()
 {
 	if (bRawStatsFile)
 	{
-		if (GetProcessingStage() != EStatsProcessingStage::SPS_Stopped)
+		if (!IsProcessingStopped())
 		{
 			SetProcessingStage( EStatsProcessingStage::SPS_ProcessStats );
 			const double StartTime = FPlatformTime::Seconds();
@@ -828,7 +829,8 @@ void FStatsReadFile::ProcessStats()
 
 			// Read all stats messages for all frames, decode callstacks.
 			const int32 FirstFrame = 0;
-			const int32 OnerPercent = FMath::Max( int32( FileInfo.TotalStatMessagesNum / 100 ), 65536 );
+			const int32 OnePercent = FMath::Max( int32( FileInfo.TotalStatMessagesNum / 200 ), 65536 );
+			int32 MessageIndexForStageProgressUpdate = 0;
 
 			for (int32 FrameIndex = 0; FrameIndex < Frames.Num(); ++FrameIndex)
 			{
@@ -854,11 +856,7 @@ void FStatsReadFile::ProcessStats()
 					for (int32 Index = 0; Index < NumStatMessages; Index++)
 					{
 						CurrentStatMessageIndex++;
-						if (CurrentStatMessageIndex % OnerPercent == 0)
-						{
-							UpdateProcessStageProgress( CurrentStatMessageIndex, FrameIndex, PacketIndex );
-						}
-
+						
 						const FStatMessage& Message = Data[Index];
 						const EStatOperation::Type Op = Message.NameAndInfo.GetField<EStatOperation>();
 						const FName RawName = Message.NameAndInfo.GetRawName();
@@ -981,25 +979,35 @@ void FStatsReadFile::ProcessStats()
 								ProcessSpecialMessageMarkerOperation( Message, *StackState );
 							}
 						}
+
+						if (CurrentStatMessageIndex > MessageIndexForStageProgressUpdate)
+						{
+							UpdateProcessStageProgress( CurrentStatMessageIndex, FrameIndex, PacketIndex );
+							MessageIndexForStageProgressUpdate += OnePercent;
+							if (IsProcessingStopped())
+							{
+								Index = NumStatMessages + 1;
+								PacketIndex = Frame.Packets.Num() + 1;
+								FrameIndex = Frames.Num() + 1;
+							}
+						}
 					}
 				}
 			}
 
-			if (GetProcessingStage() != EStatsProcessingStage::SPS_Stopped)
+			if (!IsProcessingStopped())
 			{
 				StageProgress.Set( 100 );
 
 				const double TotalTime = FPlatformTime::Seconds() - StartTime;
 				UE_LOG( LogStats, Log, TEXT( "Processing took %.2f sec(s)" ), TotalTime );
-
-				//UpdateCombinedHistoryStats();
 			}
 			else
 			{
 				UE_LOG( LogStats, Warning, TEXT( "Processing stopped, abandoning" ) );
 			}
 
-			// Clear all data.
+			// Clear all data. We shouldn't need raw stats data at this moment.
 			CombinedHistory.Empty();
 		}
 	}
@@ -1007,13 +1015,13 @@ void FStatsReadFile::ProcessStats()
 
 void FStatsReadFile::PostProcessStats()
 {
-	if (GetProcessingStage() != EStatsProcessingStage::SPS_Stopped)
+	if (!IsProcessingStopped())
 	{
 		SetProcessingStage( EStatsProcessingStage::SPS_PostProcessStats );
 	}
 }
 
-bool FStatsReadFile::UpdateReadStageProgress()
+void FStatsReadFile::UpdateReadStageProgress()
 {
 	const double CurrentSeconds = FPlatformTime::Seconds();
 	if (CurrentSeconds > LastUpdateTime + NumSecondsBetweenUpdates)
@@ -1028,11 +1036,6 @@ bool FStatsReadFile::UpdateReadStageProgress()
 	if (bShouldStopProcessing == true)
 	{
 		SetProcessingStage( EStatsProcessingStage::SPS_Stopped );
-		return false;
-	}
-	else
-	{
-		return true;
 	}
 }
 
@@ -1069,7 +1072,7 @@ void FStatsReadFile::UpdateCombinedHistoryStats()
 			CombinedHistory.Num() );
 }
 
-bool FStatsReadFile::UpdateProcessStageProgress( const int32 CurrentStatMessageIndex, const int32 FrameIndex, const int32 PacketIndex )
+void FStatsReadFile::UpdateProcessStageProgress( const int32 CurrentStatMessageIndex, const int32 FrameIndex, const int32 PacketIndex )
 {	
 	const double CurrentSeconds = FPlatformTime::Seconds();
 	if (CurrentSeconds > LastUpdateTime + NumSecondsBetweenUpdates)
@@ -1084,11 +1087,6 @@ bool FStatsReadFile::UpdateProcessStageProgress( const int32 CurrentStatMessageI
 	if (bShouldStopProcessing == true)
 	{
 		SetProcessingStage( EStatsProcessingStage::SPS_Stopped );
-		return false;
-	}
-	else
-	{
-		return true;
 	}
 }
 
diff --git a/Engine/Source/Runtime/Core/Public/Stats/StatsFile.h b/Engine/Source/Runtime/Core/Public/Stats/StatsFile.h
--- a/Engine/Source/Runtime/Core/Public/Stats/StatsFile.h
+++ b/Engine/Source/Runtime/Core/Public/Stats/StatsFile.h
@@ -878,16 +878,17 @@ struct CORE_API FStatsReadFile
 	/** Number of seconds between updating the current stage. */
 	static const double NumSecondsBetweenUpdates;
 
-public:
 	/** Creates a new reader for regular stats file. Will be nullptr for invalid files. */
 	static FStatsReadFile* CreateReaderForRegularStats( const TCHAR* Filename );
 
+public:
 	/** Reads and processes the file on the current thread. This is a blocking operation. */
 	void ReadAndProcessSynchronously();
 
 	/** Reads and processes the file using the async tasks on the pool thread. The read data is sent to the game thread using the task graph. This is a non-blocking operation. */
 	void ReadAndProcessAsynchronously();
 
+protected:
 	/** Initialization constructor. */
 	FStatsReadFile( const TCHAR* InFilename, bool bInRawStatsFile );
 
@@ -974,6 +975,14 @@ public:
 		return EStatsProcessingStage( ProcessingStage.GetValue() );
 	}
 
+	/**
+	 * @return true, if a user stopped the processing, probably using the Cancel button in the UI.
+	 */
+	const bool IsProcessingStopped() const
+	{
+		return GetProcessingStage() == EStatsProcessingStage::SPS_Stopped;
+	}
+
 	/**
 	 * @return current processing stage as string key.
 	 */
@@ -1042,20 +1051,16 @@ public:
 protected:
 	/**
 	 * Updates read stage progress periodically, does debug logging if enabled.
-	 *
-	 * @return true, if finished, false otherwise like stopped.
 	 */
-	bool UpdateReadStageProgress();
+	void UpdateReadStageProgress();
 
 	/** Dumps combined history stats. Only for raw stats. */
 	void UpdateCombinedHistoryStats();
 
 	/**
 	 * Updates process stage progress periodically, does debug logging if enabled.
-	 *
-	 * @return true, if finished, false otherwise like stopped.
 	 */
-	bool UpdateProcessStageProgress( const int32 CurrentStatMessageIndex, const int32 FrameIndex, const int32 PacketIndex );
+	void UpdateProcessStageProgress( const int32 CurrentStatMessageIndex, const int32 FrameIndex, const int32 PacketIndex );
 
 protected:
 	/** Current state of the stats. Mostly for metadata. */
