Fixed shadowed variables

diff --git a/Engine/Source/Developer/Profiler/Private/ProfilerRawStats.cpp b/Engine/Source/Developer/Profiler/Private/ProfilerRawStats.cpp
--- a/Engine/Source/Developer/Profiler/Private/ProfilerRawStats.cpp
+++ b/Engine/Source/Developer/Profiler/Private/ProfilerRawStats.cpp
@@ -229,9 +229,6 @@ void FRawStatsMemoryProfiler::PostProcessStats()
 		DumpDebugAllocations();
 	}
 
-	// We don't need the allocation map. Each snapshot has its own copy.
-	AllocationMap.Empty();
-
 	if (!IsProcessingStopped())
 	{
 		StageProgress.Set( 100 );
@@ -288,6 +285,9 @@ void FRawStatsMemoryProfiler::FreeDebugInformation()
 
 void FRawStatsMemoryProfiler::GenerateAllocationMap()
 {
+	/** Map of currently alive allocations. Ptr to AllocationInfo. */
+	TMap<uint64, FAllocationInfo> AllocationMap;
+
 	// Initialize the begin snapshot.
 	auto BeginSnapshot = SnapshotsToBeProcessed[0];
 
@@ -430,6 +430,9 @@ void FRawStatsMemoryProfiler::GenerateAllocationMap()
 	SnapshotsToBeProcessed.RemoveAt( 0 );
 	PrepareSnapshot( EndSnapshot.Value, AllocationMap );
 
+	// We don't need the allocation map. Each snapshot has its own copy.
+	AllocationMap.Empty();
+
 	SnapshotNamesArray = SnapshotNamesSet.Array();
 
 	UE_LOG( LogStats, Verbose, TEXT( "NumDuplicatedMemoryOperations: %i" ), NumDuplicatedMemoryOperations );
@@ -767,11 +770,11 @@ void FRawStatsMemoryProfiler::DumpScopedAllocations( const TCHAR* Name, const TM
 	MemoryReport.CycleRow();
 }
 
-void FRawStatsMemoryProfiler::GenerateScopedAllocations( const TMap<uint64, FAllocationInfo>& AllocationMap, TMap<FName, FCombinedAllocationInfo>& out_CombinedAllocations, uint64& TotalAllocatedMemory, uint64& NumAllocations )
+void FRawStatsMemoryProfiler::GenerateScopedAllocations( const TMap<uint64, FAllocationInfo>& InAllocationMap, TMap<FName, FCombinedAllocationInfo>& out_CombinedAllocations, uint64& TotalAllocatedMemory, uint64& NumAllocations )
 {
 	FScopeLogTime SLT( TEXT( "GenerateScopedAllocations" ), nullptr, FScopeLogTime::ScopeLog_Milliseconds );
 
-	for (const auto& It : AllocationMap)
+	for (const auto& It : InAllocationMap)
 	{
 		const FAllocationInfo& Alloc = It.Value;
 		FCombinedAllocationInfo& CombinedAllocation = out_CombinedAllocations.FindOrAdd( Alloc.EncodedCallstack );
@@ -785,7 +788,7 @@ void FRawStatsMemoryProfiler::GenerateScopedAllocations( const TMap<uint64, FAll
 	out_CombinedAllocations.ValueSort( FCombinedAllocationInfoSizeGreater() );
 }
 
-void FRawStatsMemoryProfiler::PrepareSnapshot( const FName SnapshotName, const TMap<uint64, FAllocationInfo>& AllocationMap )
+void FRawStatsMemoryProfiler::PrepareSnapshot( const FName SnapshotName, const TMap<uint64, FAllocationInfo>& InAllocationMap )
 {
 	FScopeLogTime SLT( TEXT( "PrepareSnapshot" ), nullptr, FScopeLogTime::ScopeLog_Milliseconds );
 
@@ -797,12 +800,12 @@ void FRawStatsMemoryProfiler::PrepareSnapshot( const FName SnapshotName, const T
 	}
 	SnapshotNamesSet.Add( UniqueSnapshotName );
 
-	SnapshotsWithAllocationMap.Add( UniqueSnapshotName, AllocationMap );
+	SnapshotsWithAllocationMap.Add( UniqueSnapshotName, InAllocationMap );
 
 	TMap<FName, FCombinedAllocationInfo> SnapshotCombinedAllocations;
 	uint64 TotalAllocatedMemory = 0;
 	uint64 NumAllocations = 0;
-	GenerateScopedAllocations( AllocationMap, SnapshotCombinedAllocations, TotalAllocatedMemory, NumAllocations );
+	GenerateScopedAllocations( InAllocationMap, SnapshotCombinedAllocations, TotalAllocatedMemory, NumAllocations );
 	SnapshotsWithScopedAllocations.Add( UniqueSnapshotName, SnapshotCombinedAllocations );
 
 	// Decode callstacks.
@@ -815,7 +818,7 @@ void FRawStatsMemoryProfiler::PrepareSnapshot( const FName SnapshotName, const T
 	}
 	SnapshotsWithDecodedScopedAllocations.Add( UniqueSnapshotName, SnapshotDecodedCombinedAllocations );
 
-	UE_LOG( LogStats, Warning, TEXT( "PrepareSnapshot: %s Alloc: %i Scoped: %i Total: %.2f MB" ), *UniqueSnapshotName.ToString(), AllocationMap.Num(), SnapshotCombinedAllocations.Num(), TotalAllocatedMemory / 1024.0f / 1024.0f );
+	UE_LOG( LogStats, Warning, TEXT( "PrepareSnapshot: %s Alloc: %i Scoped: %i Total: %.2f MB" ), *UniqueSnapshotName.ToString(), InAllocationMap.Num(), SnapshotCombinedAllocations.Num(), TotalAllocatedMemory / 1024.0f / 1024.0f );
 }
 
 void FRawStatsMemoryProfiler::CompareSnapshots( const FName BeginSnaphotName, const FName EndSnaphotName, TMap<FName, FCombinedAllocationInfo>& out_Result )
diff --git a/Engine/Source/Developer/Profiler/Private/ProfilerRawStats.h b/Engine/Source/Developer/Profiler/Private/ProfilerRawStats.h
--- a/Engine/Source/Developer/Profiler/Private/ProfilerRawStats.h
+++ b/Engine/Source/Developer/Profiler/Private/ProfilerRawStats.h
@@ -293,12 +293,12 @@ public:
 	void DumpScopedAllocations( const TCHAR* Name, const TMap<FString, FCombinedAllocationInfo>& CombinedAllocations );
 
 	/** Generates callstack based allocation map. */
-	void GenerateScopedAllocations( const TMap<uint64, FAllocationInfo>& AllocationMap, TMap<FName, FCombinedAllocationInfo>& out_CombinedAllocations, uint64& TotalAllocatedMemory, uint64& NumAllocations );
+	void GenerateScopedAllocations( const TMap<uint64, FAllocationInfo>& InAllocationMap, TMap<FName, FCombinedAllocationInfo>& out_CombinedAllocations, uint64& TotalAllocatedMemory, uint64& NumAllocations );
 
 	void GenerateScopedTreeAllocations( const TMap<FName, FCombinedAllocationInfo>& ScopedAllocations, FNodeAllocationInfo& out_Root );
 
 	/** Prepares data for a snapshot. */
-	void PrepareSnapshot( const FName SnapshotName, const TMap<uint64, FAllocationInfo>& AllocationMap );
+	void PrepareSnapshot( const FName SnapshotName, const TMap<uint64, FAllocationInfo>& InAllocationMap );
 
 	/** Compare two snapshots and saves the result for further processing. */
 	void CompareSnapshots( const FName BeginSnaphotName, const FName EndSnaphotName, TMap<FName, FCombinedAllocationInfo>& out_Result );
@@ -337,9 +337,6 @@ protected:
 	 */
 	TArray<FAllocationInfo> SequenceAllocationArray;
 
-	/** Map of currently alive allocations. Ptr to AllocationInfo. */
-	TMap<uint64, FAllocationInfo> AllocationMap;
-
 	/** The sequence tag mapping to the named markers. */
 	TArray<TPair<uint32, FName>> Snapshots;
 
