Profiler - Removed debug code, fixed realloc processing

diff --git a/Engine/Source/Developer/Profiler/Private/ProfilerRawStats.cpp b/Engine/Source/Developer/Profiler/Private/ProfilerRawStats.cpp
--- a/Engine/Source/Developer/Profiler/Private/ProfilerRawStats.cpp
+++ b/Engine/Source/Developer/Profiler/Private/ProfilerRawStats.cpp
@@ -192,7 +192,6 @@ void FNodeAllocationInfo::PrepareCallstackData( const TArray<FName>& InDecodedCa
 FRawStatsMemoryProfiler::FRawStatsMemoryProfiler( const TCHAR* InFilename )
 	: FStatsReadFile( InFilename, true )
 	, NumDuplicatedMemoryOperations( 0 )
-	, NumFWAMemoryOperations( 0 )
 	, NumMemoryOperations( 0 )
 	, LastSequenceTagForNamedMarker( 0 )
 {}
@@ -257,12 +256,13 @@ void FRawStatsMemoryProfiler::DumpDebugAllocations()
 	}
 
 	UE_LOG( LogStats, Warning, TEXT( "Dumping duplicated alloc map" ) );
-	const float MaxPctDisplayed = 0.80f;
+	UE_LOG( LogStats, Warning, TEXT( "TotalDuplicatedMemory: %llu bytes (%.2f MB)" ), TotalDuplicatedMemory, TotalDuplicatedMemory / 1024.0f / 1024.0f );
+	const float MaxPctDisplayed = 0.9f;
 	uint64 DisplayedSoFar = 0;
 	for (const auto& It : DuplicatedAllocMap)
 	{
 		const FAllocationInfo& Alloc = It.Value;
-		const FString AllocCallstack = FStatsCallstack::GetHumanReadable( Alloc.EncodedCallstack );
+		const FString& AllocCallstack = It.Key;
 		UE_LOG( LogStats, Log, TEXT( "%lli (%.2f MB) %s" ), Alloc.Size, Alloc.Size / 1024.0f / 1024.0f, *AllocCallstack );
 
 		DisplayedSoFar += Alloc.Size;
@@ -279,8 +279,8 @@ void FRawStatsMemoryProfiler::DumpDebugAllocations()
 
 void FRawStatsMemoryProfiler::FreeDebugInformation()
 {
-	FreeWithoutAllocMap.Empty();
 	DuplicatedAllocMap.Empty();
+	ZeroAllocMap.Empty();
 }
 
 void FRawStatsMemoryProfiler::GenerateAllocationMap()
@@ -322,107 +322,32 @@ void FRawStatsMemoryProfiler::GenerateAllocationMap()
 
 		if (Alloc.Op == EMemoryOperation::Alloc)
 		{
-			if (Alloc.Size == 0)
-			{
-				NumZeroAllocs++;
-			}
-
-			const FAllocationInfo* Found = AllocationMap.Find( Alloc.Ptr );
-
-			if (!Found)
-			{
-				AllocationMap.Add( Alloc.Ptr, Alloc );
-			}
-			else
-			{
-#if	UE_BUILD_DEBUG
-				const FAllocationInfo* FoundAndFreed = FreeWithoutAllocMap.Find( Found->Ptr );
-				const FAllocationInfo* FoundAndAllocated = FreeWithoutAllocMap.Find( Alloc.Ptr );
-
-				if (FoundAndFreed)
-				{
-					const FString FoundAndFreedCallstack = FStatsCallstack::GetHumanReadable( FoundAndFreed->EncodedCallstack );
-				}
-
-				if (FoundAndAllocated)
-				{
-					const FString FoundAndAllocatedCallstack = FStatsCallstack::GetHumanReadable( FoundAndAllocated->EncodedCallstack );
-				}
-
-				NumDuplicatedMemoryOperations++;
-
-
-				const FString FoundCallstack = FStatsCallstack::GetHumanReadable( Found->EncodedCallstack );
-				const FString AllocCallstack = FStatsCallstack::GetHumanReadable( Alloc.EncodedCallstack );
-				UE_LOG( LogStats, Warning, TEXT( "DuplicateAlloc: %s Size: %i/%i Ptr: %i/%i Tag: %i/%i" ), *AllocCallstack, Found->Size, Alloc.Size, Found->Ptr, Alloc.Ptr, Found->SequenceTag, Alloc.SequenceTag );
-#endif // UE_BUILD_DEBUG
-
-				// Replace pointer.
-				AllocationMap.Add( Alloc.Ptr, Alloc );
-				// Store the old pointer.
-				DuplicatedAllocMap.Add( Alloc.Ptr, *Found );
-			}
+			ProcessAlloc( Alloc, AllocationMap );
 		}
 		else if (Alloc.Op == EMemoryOperation::Realloc)
 		{
-			if (Alloc.Size == 0)
-			{
-				NumZeroAllocs++;
-			}
-
 			// Previous Alloc or Realloc
-			const FAllocationInfo* FoundOld = AllocationMap.Find( Alloc.OldPtr );
-
-			if (FoundOld)
+			if (Alloc.OldPtr != 0)
 			{
-				const bool bIsValid = Alloc.SequenceTag > FoundOld->SequenceTag;
-				if (!bIsValid)
-				{
-					UE_LOG( LogStats, Warning, TEXT( "InvalidRealloc Ptr: %llu, Seq: %i/%i" ), Alloc.Ptr, Alloc.SequenceTag, FoundOld->SequenceTag );
-				}
-				// Remove the old allocation.
-				AllocationMap.Remove( Alloc.OldPtr );
-				AllocationMap.Add( Alloc.Ptr, Alloc );
+				ProcessFree( Alloc, AllocationMap, true );
 			}
-			// If we have situation where the old pointer is the same as the new pointer
-			// There is high change that the original call looked like this 
-			// Ptr = Malloc( 40 ); Ptr = Realloc( Ptr, 40 );
-			else if (Alloc.OldPtr != Alloc.Ptr)
-			{
-				// No OldPtr, should not happen as it means realloc without initial alloc.
-				// Or Realloc after Malloc(0)
 
 #if	UE_BUILD_DEBUG
+			if (Alloc.OldPtr == 0 && Alloc.Size == 0)
+			{
 				const FString ReallocCallstack = FStatsCallstack::GetHumanReadable( Alloc.EncodedCallstack );
-				UE_LOG( LogStats, VeryVerbose, TEXT( "ReallocWithoutAlloc: %s %i %i/%i [%i]" ), *ReallocCallstack, Alloc.Size, Alloc.OldPtr, Alloc.Ptr, Alloc.SequenceTag );
-#endif // UE_BUILD_DEBUG
-
-				AllocationMap.Add( Alloc.Ptr, Alloc );
+				UE_LOG( LogStats, VeryVerbose, TEXT( "ReallocZero: %s %i %i/%i [%i]" ), *ReallocCallstack, Alloc.Size, Alloc.OldPtr, Alloc.Ptr, Alloc.SequenceTag );
 			}
+#endif // UE_BUILD_DEBUG
 
+			if (Alloc.Ptr != 0)
+			{
+				ProcessAlloc( Alloc, AllocationMap );
+			}			
 		}
 		else if (Alloc.Op == EMemoryOperation::Free)
 		{
-			const FAllocationInfo* Found = AllocationMap.Find( Alloc.Ptr );
-			if (Found)
-			{
-				const bool bIsValid = Alloc.SequenceTag > Found->SequenceTag;
-				if (!bIsValid)
-				{
-					UE_LOG( LogStats, Warning, TEXT( "InvalidFree Ptr: %llu, Seq: %i/%i" ), Alloc.Ptr, Alloc.SequenceTag, Found->SequenceTag );
-				}
-				AllocationMap.Remove( Alloc.Ptr );
-			}
-			else
-			{
-				FreeWithoutAllocMap.Add( Alloc.Ptr, Alloc );
-				NumFWAMemoryOperations++;
-
-#if	UE_BUILD_DEBUG
-				const FString FWACallstack = FStatsCallstack::GetHumanReadable( Alloc.EncodedCallstack );
-				UE_LOG( LogStats, VeryVerbose, TEXT( "FreeWithoutAllocCallstack: %s %i" ), *FWACallstack, Alloc.Ptr );
-#endif // UE_BUILD_DEBUG
-			}
+			ProcessFree( Alloc, AllocationMap, false );
 		}
 	}
 
@@ -436,10 +361,66 @@ void FRawStatsMemoryProfiler::GenerateAllocationMap()
 	SnapshotNamesArray = SnapshotNamesSet.Array();
 
 	UE_LOG( LogStats, Verbose, TEXT( "NumDuplicatedMemoryOperations: %i" ), NumDuplicatedMemoryOperations );
-	UE_LOG( LogStats, Verbose, TEXT( "NumFWAMemoryOperations:        %i" ), NumFWAMemoryOperations );
 	UE_LOG( LogStats, Verbose, TEXT( "NumZeroAllocs:                 %i" ), NumZeroAllocs );
 }
 
+void FRawStatsMemoryProfiler::ProcessAlloc( const FAllocationInfo& AllocInfo, TMap<uint64, FAllocationInfo>& AllocationMap )
+{
+	if( AllocInfo.Size == 0 )
+	{
+		NumZeroAllocs++;
+		ZeroAllocMap.Add( FStatsCallstack::GetHumanReadable( AllocInfo.EncodedCallstack ), AllocInfo );
+	}
+
+	const FAllocationInfo* Found = AllocationMap.Find( AllocInfo.Ptr );
+	if (!Found)
+	{
+		AllocationMap.Add( AllocInfo.Ptr, AllocInfo );
+	}
+	else
+	{
+		NumDuplicatedMemoryOperations++;
+
+#if	UE_BUILD_DEBUG
+		const FString FoundCallstack = FStatsCallstack::GetHumanReadable( Found->EncodedCallstack );
+		const FString AllocCallstack = FStatsCallstack::GetHumanReadable( AllocInfo.EncodedCallstack );
+		UE_LOG( LogStats, VeryVerbose, TEXT( "DuplicatedAlloc" ) );
+		UE_LOG( LogStats, VeryVerbose, TEXT( "FoundCallstack: %s [%s]" ), *FoundCallstack, Found->Op==EMemoryOperation::Alloc ? TEXT("Alloc") : TEXT("Realloc") );
+		UE_LOG( LogStats, VeryVerbose, TEXT( "AllocCallstack: %s [%s]" ), *AllocCallstack, AllocInfo.Op==EMemoryOperation::Alloc ? TEXT("Alloc") : TEXT("Realloc") );
+		UE_LOG( LogStats, VeryVerbose, TEXT( "Size: %i/%i Ptr: %llu/%llu Tag: %i/%i" ), Found->Size, AllocInfo.Size, Found->Ptr, AllocInfo.Ptr, Found->SequenceTag, AllocInfo.SequenceTag );
+
+		// Store the old pointer.
+		DuplicatedAllocMap.Add( FoundCallstack, *Found );
+#endif // UE_BUILD_DEBUG
+
+		// Replace pointer.
+		AllocationMap.Add( AllocInfo.Ptr, AllocInfo );
+	}
+}
+
+void FRawStatsMemoryProfiler::ProcessFree( const FAllocationInfo& FreeInfo, TMap<uint64, FAllocationInfo>& AllocationMap, const bool bReallocFree )
+{
+	// bReallocFree is not needed here, but it's easier to read the code.
+	const uint64 PtrToBeFreed = bReallocFree ? FreeInfo.OldPtr : FreeInfo.Ptr;
+	const FAllocationInfo* Found = AllocationMap.Find( PtrToBeFreed );
+	if (Found)
+	{
+		const bool bIsValid = FreeInfo.SequenceTag > Found->SequenceTag;
+		if (!bIsValid)
+		{
+			UE_LOG( LogStats, Warning, TEXT( "InvalidFree Ptr: %llu, Seq: %i/%i" ), PtrToBeFreed, FreeInfo.SequenceTag, Found->SequenceTag );
+		}
+		AllocationMap.Remove( PtrToBeFreed );
+	}
+	else
+	{
+#if	UE_BUILD_DEBUG
+		const FString FWACallstack = FStatsCallstack::GetHumanReadable( FreeInfo.EncodedCallstack );
+		UE_LOG( LogStats, VeryVerbose, TEXT( "FreeWithoutAlloc: %s, %llu" ), *FWACallstack, PtrToBeFreed );
+#endif // UE_BUILD_DEBUG
+	}
+}
+
 void FRawStatsMemoryProfiler::UpdateGenerateMemoryMapProgress( const int32 AllocationIndex )
 {
 	const double CurrentSeconds = FPlatformTime::Seconds();
diff --git a/Engine/Source/Developer/Profiler/Private/ProfilerRawStats.h b/Engine/Source/Developer/Profiler/Private/ProfilerRawStats.h
--- a/Engine/Source/Developer/Profiler/Private/ProfilerRawStats.h
+++ b/Engine/Source/Developer/Profiler/Private/ProfilerRawStats.h
@@ -280,6 +280,12 @@ protected:
 	 */
 	void GenerateAllocationMap();
 
+	/** Processes an alloc memory operation. */
+	void ProcessAlloc( const FAllocationInfo& AllocInfo, TMap<uint64, FAllocationInfo>& AllocationMap );
+
+	/** Processes a free memory operation. */
+	void ProcessFree( const FAllocationInfo& FreeInfo, TMap<uint64, FAllocationInfo>& AllocationMap, const bool bReallocFree );
+
 	/** Dumps debug allocations. */
 	void DumpDebugAllocations();
 
@@ -357,18 +363,15 @@ protected:
 
 	TMap<FName, TMap<FString, FCombinedAllocationInfo> > SnapshotsWithDecodedScopedAllocations;
 
-	/** Free without allocation map. Debug only. */
-	TMultiMap<uint64, FAllocationInfo> FreeWithoutAllocMap;
-
 	/** Duplicated allocation map. Debug only. */
-	TMultiMap<uint64, FAllocationInfo> DuplicatedAllocMap;
+	TMultiMap<FString, FAllocationInfo> DuplicatedAllocMap;
+
+	/** Zero allocation map. Debug only. */
+	TMultiMap<FString, FAllocationInfo> ZeroAllocMap;
 
 	/** Number of duplicated memory operations. */
 	int32 NumDuplicatedMemoryOperations;
 
-	/** Number of Free without Alloc. */
-	int32 NumFWAMemoryOperations;
-
 	/** Number of Malloc(0). */
 	int32 NumZeroAllocs = 0;
 
