Profiler - Updated DumpMemoryCommand

diff --git a/Engine/Source/Developer/Profiler/Private/ProfilerModule.cpp b/Engine/Source/Developer/Profiler/Private/ProfilerModule.cpp
--- a/Engine/Source/Developer/Profiler/Private/ProfilerModule.cpp
+++ b/Engine/Source/Developer/Profiler/Private/ProfilerModule.cpp
@@ -72,7 +72,7 @@ void FProfilerModule::StatsMemoryDumpCommand( const TCHAR* Filename )
 	TUniquePtr<FRawStatsMemoryProfiler> Instance( FCreateStatsReader<FRawStatsMemoryProfiler>::ForRawStats( Filename ) );
 	if (Instance)
 	{
-		Instance->ReadAndProcessAsynchronously();
+		Instance->ReadAndProcessSynchronously();
 
 		while (Instance->IsBusy())
 		{
