Profiler: Added limited support for text based filtering of the event list (doesn't filter in all modes, needs to be reapplied when the current frame changes)
[INTEGRATE] Change 2671883 by [REDACTED_PERSON][REDACTED_HANDLE] on 2015/08/27 21:11:41

diff --git a/Engine/Source/Developer/Profiler/Private/Widgets/SEventGraph.cpp b/Engine/Source/Developer/Profiler/Private/Widgets/SEventGraph.cpp
--- a/Engine/Source/Developer/Profiler/Private/Widgets/SEventGraph.cpp
+++ b/Engine/Source/Developer/Profiler/Private/Widgets/SEventGraph.cpp
@@ -1267,9 +1267,9 @@ TSharedRef<SWidget> SEventGraph::GetWidgetBoxForOptions()
 		.Padding( 1.0f )
 		[
 			SAssignNew( FilteringSearchBox, SSearchBox )
-			.Visibility( EVisibility::Collapsed )
 			.HintText( LOCTEXT("FilteringSearchBox_HintText", "Search or filter event(s)") )
 			.OnTextChanged( this, &SEventGraph::FilteringSearchBox_OnTextChanged )
+			.OnTextCommitted( this, &SEventGraph::FilteringSearchBox_OnTextCommitted )
 			.IsEnabled( this, &SEventGraph::FilteringSearchBox_IsEnabled )
 			.ToolTipText( LOCTEXT("FilteringSearchBox_TT", "Type here to search or filter events") )
 		]
@@ -1526,9 +1526,43 @@ void SEventGraph::FilteringSearchBox_OnTextChanged( const FText& InFilterText )
 {
 }
 
+static bool RecursiveShowUnfilteredItems(TSharedPtr< STreeView<FEventGraphSamplePtr> >& TreeView, TArray<FEventGraphSamplePtr>& Nodes)
+{
+	bool bExpandedAnyChildren = false;
+
+	for (FEventGraphSamplePtr& Node : Nodes)
+	{
+		const bool bChildIsExpanded = RecursiveShowUnfilteredItems(TreeView, Node->GetChildren());
+		const bool bThisWantsExpanded = !Node->PropertyValueAsBool(EEventPropertyIndex::bIsFiltered);
+		const bool bExpandThis = bChildIsExpanded || bThisWantsExpanded;
+		bExpandedAnyChildren |= bExpandThis;
+
+		TreeView->SetItemExpansion(Node, bExpandThis);
+	}
+
+	return bExpandedAnyChildren;
+}
+
+void SEventGraph::FilteringSearchBox_OnTextCommitted(const FText& NewText, ETextCommit::Type CommitType)
+{
+	PROFILER_SCOPE_LOG_TIME(TEXT("SEventGraph::FilterOutByText_Execute"), nullptr);
+
+	SaveCurrentEventGraphState();
+	FEventGraphState* Op = GetCurrentState()->CreateCopyWithTextFiltering(NewText.ToString());
+	CurrentStateIndex = EventGraphStatesHistory.Insert(MakeShareable(Op), CurrentStateIndex + 1);
+	RestoreEventGraphStateFrom(GetCurrentState());
+
+	// Auto-expand to view the unfiltered items
+	if (GetCurrentStateViewMode() == EEventGraphViewModes::Hierarchical)
+	{
+		RecursiveShowUnfilteredItems(TreeView_Base, GetCurrentState()->GetRoot()->GetChildren());
+		TreeView_Refresh();
+	}
+}
+
 bool SEventGraph::FilteringSearchBox_IsEnabled() const
 {
-	return false;
+	return true;
 }
 
 TSharedPtr<SWidget> SEventGraph::EventGraph_GetMenuContent() const
@@ -2003,6 +2037,9 @@ void SEventGraph::SetTreeItemsForViewMode( const EEventGraphViewModes::Type NewV
 	GetCurrentState()->ViewMode = NewViewMode;
 	GetCurrentState()->EventGraphType = NewEventGraphType;
 
+	GetCurrentState()->ApplyCulling();
+	GetCurrentState()->ApplyFiltering();
+
 	if( GetCurrentStateViewMode() == EEventGraphViewModes::Hierarchical )
 	{
 		TreeView_Base->SetTreeItemsSource( &GetCurrentState()->GetRoot()->GetChildren() );
diff --git a/Engine/Source/Developer/Profiler/Private/Widgets/SEventGraph.h b/Engine/Source/Developer/Profiler/Private/Widgets/SEventGraph.h
--- a/Engine/Source/Developer/Profiler/Private/Widgets/SEventGraph.h
+++ b/Engine/Source/Developer/Profiler/Private/Widgets/SEventGraph.h
@@ -255,6 +255,7 @@ protected:
 	bool EventGraphType_IsEnabled( const EEventGraphTypes::Type InEventGraphType ) const;
 
 	void FilteringSearchBox_OnTextChanged( const FText& InFilterText );
+	void FilteringSearchBox_OnTextCommitted( const FText& NewText, ETextCommit::Type CommitType );
 	bool FilteringSearchBox_IsEnabled() const;
 
 	TSharedPtr<SWidget> EventGraph_GetMenuContent() const;
@@ -607,6 +608,12 @@ protected:
 			return new FEventGraphState( *this, InFilterPropertyName, InFilterEventPtr, FFilteredTag() );
 		}
 
+		FEventGraphState* CreateCopyWithTextFiltering(const FString& InFilterText)
+		{
+			FEventGraphState* Result = new FEventGraphState( *this, NAME_None, nullptr, FFilteredTag() );
+			InFilterText.ParseIntoArray(Result->TextBasedFilterStringTokens, TEXT(" "));
+			return Result;
+		}
 	protected:
 		/** Copy constructor for setting new root. */
 		FEventGraphState( const FEventGraphState& EventGraphState, const TArray<FEventGraphSamplePtr>& UniqueLeafs )
@@ -704,6 +711,9 @@ protected:
 		/** Value of the property used to filter out the event graph. */
 		FEventGraphSamplePtr FilterEventPtr;
 
+		/** Text substrings to match for text-based filtering (AND - all must be present) */
+		TArray<FString> TextBasedFilterStringTokens;
+
 		const double CreationTime;
 		EEventHistoryTypes::Type HistoryType;
 
@@ -793,9 +803,58 @@ protected:
 			}
 		}
 
+private:
+		static bool PassesTokenFilter(const TArray<FString>& FilterTokens, const FString& TestString)
+		{
+			for (const FString& Token : FilterTokens)
+			{
+				if (!TestString.Contains(Token))
+				{
+					return false;
+				}
+			}
+
+			return true;
+		}
+
+		// Sets the filter and optionally culled properties, returning true if any child passed the filter
+		bool ApplyTextBasedFilterInternal(TArray< FEventGraphSamplePtr >& Nodes, bool bCullAsWell)
+		{
+			bool bAnyPasses = false;
+
+			for (FEventGraphSamplePtr& Node : Nodes)
+			{
+				const bool bChildPassesFilter = ApplyTextBasedFilterInternal(Node->GetChildren(), bCullAsWell);
+
+				const bool bThisPassesFilter = PassesTokenFilter(TextBasedFilterStringTokens, Node->_StatName.ToString());
+
+				bAnyPasses = bAnyPasses || bThisPassesFilter || bChildPassesFilter;
+
+				Node->PropertyValueAsBool(EEventPropertyIndex::bIsFiltered) = !bThisPassesFilter;
+				
+				if (bCullAsWell)
+				{
+					const bool bChildSavesFromCull = (ViewMode == EEventGraphViewModes::Hierarchical) && bChildPassesFilter;
+
+					bool& bCullState = Node->PropertyValueAsBool(EEventPropertyIndex::bIsCulled);
+
+					bCullState = (bCullState || !bChildPassesFilter) && !bChildSavesFromCull;
+				}
+			}
+
+			return bAnyPasses;
+		}
+
+public:
 		void ApplyFiltering()
 		{
-			if( IsFiltered() )
+			if (TextBasedFilterStringTokens.Num() > 0)
+			{
+				// Apply text substring filtering (and optionally culling).
+				const bool bAlsoCull = true;
+				ApplyTextBasedFilterInternal(GetRoot()->GetChildren(), bAlsoCull);
+			}
+			else if (IsFiltered())
 			{
 				// Apply filtering.
 				FEventArrayBooleanOp::ExecuteOperation
