SessionFrontend: Layout tweaks

diff --git a/Engine/Source/Developer/AutomationWindow/Private/SAutomationWindow.cpp b/Engine/Source/Developer/AutomationWindow/Private/SAutomationWindow.cpp
--- a/Engine/Source/Developer/AutomationWindow/Private/SAutomationWindow.cpp
+++ b/Engine/Source/Developer/AutomationWindow/Private/SAutomationWindow.cpp
@@ -230,7 +230,7 @@ void SAutomationWindow::Construct( const FArguments& InArgs, const IAutomationCo
 				]
 				+SVerticalBox::Slot()
 				.FillHeight(1.0f)
-				.Padding(0.0f, 8.0f, 0.0f, 0.0f)
+				.Padding(0.0f, 4.0f, 0.0f, 0.0f)
 				[
 					SNew(SOverlay)
 					+SOverlay::Slot()
diff --git a/Engine/Source/Developer/Profiler/Private/Widgets/SProfilerWindow.cpp b/Engine/Source/Developer/Profiler/Private/Widgets/SProfilerWindow.cpp
--- a/Engine/Source/Developer/Profiler/Private/Widgets/SProfilerWindow.cpp
+++ b/Engine/Source/Developer/Profiler/Private/Widgets/SProfilerWindow.cpp
@@ -138,6 +138,7 @@ void SProfilerWindow::Construct( const FArguments& InArgs )
 				/** Profiler Mini-view. */
 				+ SVerticalBox::Slot()
 				.AutoHeight()
+				.Padding(0.0f, 4.0f, 0.0f, 0.0f)
 				[
 					SNew(SBox)
 					.HeightOverride( 48.0f )
@@ -156,16 +157,9 @@ void SProfilerWindow::Construct( const FArguments& InArgs )
 					]
 				]
 
-				+ SVerticalBox::Slot()
-				.AutoHeight()
-				[
-					SNew( SSpacer )
-					.Size( FVector2D( 2.0f, 2.0f ) )
-				]
-
 				+SVerticalBox::Slot()
 				.FillHeight( 1.0f )
-				.Padding(0.0f, 8.0f, 0.0f, 0.0f)
+				.Padding(0.0f, 4.0f, 0.0f, 0.0f)
 				[
 					SNew( SHorizontalBox )
 					.IsEnabled( this, &SProfilerWindow::IsProfilerEnabled )
diff --git a/Engine/Source/Developer/SessionFrontend/Private/Widgets/Console/SSessionConsole.cpp b/Engine/Source/Developer/SessionFrontend/Private/Widgets/Console/SSessionConsole.cpp
--- a/Engine/Source/Developer/SessionFrontend/Private/Widgets/Console/SSessionConsole.cpp
+++ b/Engine/Source/Developer/SessionFrontend/Private/Widgets/Console/SSessionConsole.cpp
@@ -51,7 +51,7 @@ void SSessionConsole::Construct( const FArguments& InArgs, ISessionManagerRef In
 
 				+ SVerticalBox::Slot()
 					.AutoHeight()
-					.Padding(0.0f, 8.0f, 0.0f, 0.0f)
+					.Padding(0.0f, 4.0f, 0.0f, 0.0f)
 					[
 						// filter bar
 						SNew(SExpandableArea)
