Profiler:made filter box resizable; replaced spaces with slot paddings; more sensible default sizes.

diff --git a/Engine/Source/Developer/Profiler/Private/Widgets/SEventGraph.cpp b/Engine/Source/Developer/Profiler/Private/Widgets/SEventGraph.cpp
--- a/Engine/Source/Developer/Profiler/Private/Widgets/SEventGraph.cpp
+++ b/Engine/Source/Developer/Profiler/Private/Widgets/SEventGraph.cpp
@@ -756,8 +756,8 @@ void SEventGraph::Construct( const FArguments& InArgs )
 {
 	static TArray<FEventGraphSamplePtr> StaticEventArray;
 
-	SAssignNew( ExternalScrollbar, SScrollBar )
-		.AlwaysShowScrollbar( true );
+	SAssignNew(ExternalScrollbar, SScrollBar)
+		.AlwaysShowScrollbar(true);
 
 	ChildSlot
 	[
@@ -765,18 +765,18 @@ void SEventGraph::Construct( const FArguments& InArgs )
 			.Orientation(Orient_Vertical)
 
 		+ SSplitter::Slot()
-			.Value(0.7f)
+			.Value(0.5f)
 			[
-				SNew( SVerticalBox )
+				SNew(SVerticalBox)
 
 				+ SVerticalBox::Slot()
 					.AutoHeight()
 					[
-						SNew( SBorder )
-							.BorderImage( FEditorStyle::GetBrush("ToolPanel.GroupBorder") )
-							.Padding( 2.0f )
+						SNew(SBorder)
+							.BorderImage(FEditorStyle::GetBrush("ToolPanel.GroupBorder"))
+							.Padding(2.0f)
 							[
-								SNew( SHorizontalBox )
+								SNew(SHorizontalBox)
 
 								+ SHorizontalBox::Slot()
 									.AutoWidth()
@@ -786,76 +786,57 @@ void SEventGraph::Construct( const FArguments& InArgs )
 
 								+ SHorizontalBox::Slot()
 									.AutoWidth()
-									[
-										SNew( SSpacer )
-											.Size( FVector2D( 2.0f, 2.0f ) )
-									]
-
-								+ SHorizontalBox::Slot()
-									.AutoWidth()
+									.Padding(2.0f, 0.0f, 0.0f, 0.0f)
 									[
 										GetWidgetForEventGraphViewModes()
 									]
 
 								+ SHorizontalBox::Slot()
-									.AutoWidth()
-									[
-										SNew( SSpacer )
-											.Size( FVector2D( 2.0f, 2.0f ) )
-									]
-
-								+ SHorizontalBox::Slot()
-									.FillWidth( 1.0f )
+									.FillWidth(1.0f)
+									.Padding(2.0f, 0.0f, 0.0f, 0.0f)
 									[
 										GetWidgetBoxForOptions()
 									]
 							]
 					]
  
-				// Spacer
-				+ SVerticalBox::Slot()
-					.AutoHeight()
-					[
-						SNew( SSpacer )
-							.Size( FVector2D( 2.0f, 2.0f ) )
-					]
-
 				// Function details view ( @see VS2012 profiler )
 				+ SVerticalBox::Slot()
-					.AutoHeight()
+					.FillHeight(1.0f)
+					.Padding(0.0f, 2.0f, 0.0f, 0.0f)
 					[		
 						SAssignNew(FunctionDetailsBox,SBox)
-							.HeightOverride( 224.0f )
+							.HeightOverride(224.0f)
 							[
-								SNew( SBorder )
-									.BorderImage( FEditorStyle::GetBrush("ToolPanel.GroupBorder") )
-									.Padding( 2.0f )
+								SNew(SBorder)
+									.BorderImage(FEditorStyle::GetBrush("ToolPanel.GroupBorder"))
+									.Padding(2.0f)
 									[
-										SNew( SHorizontalBox )
+										SNew(SHorizontalBox)
 
 										// Calling Functions
 										+ SHorizontalBox::Slot()
-											.FillWidth( 1.0f )
-											.Padding( 2.0f )
+											.FillWidth(1.0f)
+											.Padding(2.0f)
 											[
-												GetVerticalBoxForFunctionDetails( VerticalBox_TopCalling, LOCTEXT("FunctionDetails_CallingFunctions","Calling Functions") )
+												GetVerticalBoxForFunctionDetails(VerticalBox_TopCalling, LOCTEXT("FunctionDetails_CallingFunctions","Calling Functions"))
 											]
 
 										// ->
 										+ SHorizontalBox::Slot()
 											.AutoWidth()
-											.HAlign( HAlign_Center )
-											.VAlign( VAlign_Top )
-											.Padding( 2.0f )
+											.HAlign(HAlign_Center)
+											.VAlign(VAlign_Top)
+											.Padding(2.0f)
 											[
-												SNew( SImage )
-													.Image( FEditorStyle::GetBrush("Animation.Forward") )
+												SNew(SImage)
+													.Image(FEditorStyle::GetBrush("Animation.Forward"))
 											]
 
 										// Current Function
 										+ SHorizontalBox::Slot()
-											.FillWidth( 1.0f )
-											.Padding( 2.0f )
+											.FillWidth(1.0f)
+											.Padding(2.0f)
 											[
 												GetVerticalBoxForCurrentFunction()
 											]
@@ -863,18 +844,18 @@ void SEventGraph::Construct( const FArguments& InArgs )
 										// ->
 										+ SHorizontalBox::Slot()
 											.AutoWidth()
-											.HAlign( HAlign_Center )
-											.VAlign( VAlign_Top )
-											.Padding( 2.0f )
+											.HAlign(HAlign_Center)
+											.VAlign(VAlign_Top)
+											.Padding(2.0f)
 											[
-												SNew( SImage )
-													.Image( FEditorStyle::GetBrush("Animation.Forward") )
+												SNew(SImage)
+													.Image(FEditorStyle::GetBrush("Animation.Forward"))
 											]
 
 										// Called Functions
 										+ SHorizontalBox::Slot()
-											.FillWidth( 1.0f )
-											.Padding( 2.0f )
+											.FillWidth(1.0f)
+											.Padding(2.0f)
 											[
 												GetVerticalBoxForFunctionDetails( VerticalBox_TopCalled, LOCTEXT("FunctionDetails_CalledFunctions","Called Functions") )
 											]
@@ -884,31 +865,30 @@ void SEventGraph::Construct( const FArguments& InArgs )
 			]
 
 		+ SSplitter::Slot()
-			.Value(0.3f)
+			.Value(0.5f)
 			[
-				SNew( SBorder )
+				SNew(SBorder)
 					.BorderImage( FEditorStyle::GetBrush("ToolPanel.GroupBorder") )
 					.Padding(0.0f)
 					[
 						SNew(SHorizontalBox)
 
 						+SHorizontalBox::Slot()
-							.FillWidth( 1.0f )
+							.FillWidth(1.0f)
 							[
-								SNew( SHorizontalBox )
+								SNew(SHorizontalBox)
 
 								+ SHorizontalBox::Slot()
-									//.Padding( 2.0f )
 									[
-										SAssignNew( TreeView_Base, STreeView< FEventGraphSamplePtr > )
-											.ExternalScrollbar( ExternalScrollbar )
-											.SelectionMode( ESelectionMode::Multi )
-											.TreeItemsSource( &StaticEventArray )
-											.OnGetChildren( this, &SEventGraph::EventGraph_OnGetChildren )
-											.OnGenerateRow( this, &SEventGraph::EventGraph_OnGenerateRow )
-											.OnSelectionChanged( this, &SEventGraph::EventGraph_OnSelectionChanged )
-											.OnContextMenuOpening( FOnContextMenuOpening::CreateSP( this, &SEventGraph::EventGraph_GetMenuContent ) )
-											.ItemHeight( 12 )
+										SAssignNew(TreeView_Base, STreeView<FEventGraphSamplePtr>)
+											.ExternalScrollbar(ExternalScrollbar)
+											.SelectionMode(ESelectionMode::Multi)
+											.TreeItemsSource(&StaticEventArray)
+											.OnGetChildren(this, &SEventGraph::EventGraph_OnGetChildren)
+											.OnGenerateRow(this, &SEventGraph::EventGraph_OnGenerateRow)
+											.OnSelectionChanged(this, &SEventGraph::EventGraph_OnSelectionChanged)
+											.OnContextMenuOpening(FOnContextMenuOpening::CreateSP(this, &SEventGraph::EventGraph_GetMenuContent))
+											.ItemHeight(12.0f)
 											.HeaderRow
 											(
 												SAssignNew(TreeViewHeaderRow,SHeaderRow)
@@ -920,8 +900,8 @@ void SEventGraph::Construct( const FArguments& InArgs )
 						+ SHorizontalBox::Slot()
 							.AutoWidth()
 							[
-								SNew( SBox )
-									.WidthOverride( FOptionalSize( 16 ) )
+								SNew(SBox)
+									.WidthOverride(FOptionalSize(16.0f))
 									[
 										ExternalScrollbar.ToSharedRef()
 									]
diff --git a/Engine/Source/Developer/Profiler/Private/Widgets/SProfilerWindow.cpp b/Engine/Source/Developer/Profiler/Private/Widgets/SProfilerWindow.cpp
--- a/Engine/Source/Developer/Profiler/Private/Widgets/SProfilerWindow.cpp
+++ b/Engine/Source/Developer/Profiler/Private/Widgets/SProfilerWindow.cpp
@@ -109,197 +109,176 @@ void SProfilerWindow::Construct( const FArguments& InArgs )
 			SNew(SOverlay)
 
 			// Overlay slot for the main profiler window area, the first
-			+SOverlay::Slot()
-			[
-				SAssignNew(MainContentPanel,SVerticalBox)
-
-				+SVerticalBox::Slot()
-				.AutoHeight()
-				[
-					SNew( SHorizontalBox )
-
-					+SHorizontalBox::Slot()
-					.FillWidth( 1.0f )
-					.Padding( 0.0f )
-					.HAlign(HAlign_Fill)
-					.VAlign(VAlign_Fill)
-					[
-						SNew(SProfilerToolbar)
-					]
-				]
-
-				+SVerticalBox::Slot()
-				.AutoHeight()
+			+ SOverlay::Slot()
 				[
-					SNew( SSpacer )
-					.Size( FVector2D( 2.0f, 2.0f ) )
-				]
+					SAssignNew(MainContentPanel, SVerticalBox)
 
-				/** Profiler Mini-view. */
-				+ SVerticalBox::Slot()
-				.AutoHeight()
-				.Padding(0.0f, 4.0f, 0.0f, 0.0f)
-				[
-					SNew(SBox)
-					.HeightOverride( 48.0f )
-					[
-						SNew( SHorizontalBox )
-						.IsEnabled( this, &SProfilerWindow::IsProfilerEnabled )
-
-						+ SHorizontalBox::Slot()
-						.FillWidth( 1.0f )
-						.Padding( 0.0f )
-						.HAlign( HAlign_Fill )
-						.VAlign( VAlign_Fill )
+					+ SVerticalBox::Slot()
+						.AutoHeight()
 						[
-							SAssignNew( ProfilerMiniView, SProfilerMiniView )
-						]
-					]
-				]
-
-				+SVerticalBox::Slot()
-				.FillHeight( 1.0f )
-				.Padding(0.0f, 4.0f, 0.0f, 0.0f)
-				[
-					SNew( SHorizontalBox )
-					.IsEnabled( this, &SProfilerWindow::IsProfilerEnabled )
-
-					+SHorizontalBox::Slot()
-					.AutoWidth()
-					.Expose( FiltersAndPresetsSlot )
-					[
-						SNew(SBox)
-						.WidthOverride( 256.0f )
-						[
-							SNew(SVerticalBox)
-
-							// Header
-							+SVerticalBox::Slot()
-							.AutoHeight()
-							[
-								SNew(SHorizontalBox)
+							SNew(SHorizontalBox)
 
-								+SHorizontalBox::Slot()
-								.AutoWidth()
+							+ SHorizontalBox::Slot()
+								.FillWidth(1.0f)
+								.Padding(0.0f)
+								.HAlign(HAlign_Fill)
+								.VAlign(VAlign_Fill)
 								[
-									SNew(SImage)
-									.Image( FEditorStyle::GetBrush( TEXT("Profiler.Tab.FiltersAndPresets") ) )
+									SNew(SProfilerToolbar)
 								]
+						]
 
-								+SHorizontalBox::Slot()
-								.AutoWidth()
+					/** Profiler Mini-view. */
+					+ SVerticalBox::Slot()
+						.AutoHeight()
+						.Padding(0.0f, 6.0f, 0.0f, 0.0f)
+						[
+							SNew(SBox)
+								.HeightOverride(48.0f)
 								[
-									SNew(STextBlock)
-									.Text( LOCTEXT("FiltersAndPresetsLabel", "Filters And Presets") )
+									SNew(SHorizontalBox)
+										.IsEnabled(this, &SProfilerWindow::IsProfilerEnabled)
+
+										+ SHorizontalBox::Slot()
+											.FillWidth(1.0f)
+											.Padding(0.0f)
+											.HAlign(HAlign_Fill)
+											.VAlign(VAlign_Fill)
+											[
+												SAssignNew(ProfilerMiniView, SProfilerMiniView)
+											]
 								]
-							]
-
-							// Spacer
-							+SVerticalBox::Slot()
-							.AutoHeight()
-							[
-								SNew( SSpacer )
-								.Size( FVector2D( 2.0f, 2.0f ) )
-							]
-
-							// Filters And Presets
-							+SVerticalBox::Slot()
-							.FillHeight( 1.0f )
-							[
-								SAssignNew( FiltersAndPresets, SFiltersAndPresets )
-							]
 						]
-					]
 
-					+SHorizontalBox::Slot()
-					.FillWidth( 1.0f )
-					.Padding(6.0f, 0.0f, 0.0f, 0.0f)
-					[
-						SNew( SSplitter )
-						.Orientation( Orient_Vertical )
-						//.PhysicalSplitterHandleSize( 2.0f )
-						//.HitDetectionSplitterHandleSize( 4.0f )
-
-						+SSplitter::Slot()
-						.Value( 1.0f )
+					+ SVerticalBox::Slot()
+						.FillHeight(1.0f)
+						.Padding(0.0f, 4.0f, 0.0f, 0.0f)
 						[
-							SNew(SVerticalBox)
-
-							// Header
-							+SVerticalBox::Slot()
-							.AutoHeight()
-							[
-								SNew(SHorizontalBox)
+							SNew(SSplitter)
+								.IsEnabled( this, &SProfilerWindow::IsProfilerEnabled )
+								.Orientation(Orient_Horizontal)
 
-								+SHorizontalBox::Slot()
-								.AutoWidth()
+							+ SSplitter::Slot()
+								.Expose(FiltersAndPresetsSlot)
+								.Value(0.25f)
 								[
-									SNew(SImage)
-									.Image( FEditorStyle::GetBrush( TEXT("Profiler.Tab.GraphView") ) )
+									SNew(SBox)
+										.WidthOverride(256.0f)
+										[
+											SNew(SVerticalBox)
+
+											// Header
+											+ SVerticalBox::Slot()
+												.AutoHeight()
+												[
+													SNew(SHorizontalBox)
+
+													+ SHorizontalBox::Slot()
+														.AutoWidth()
+														[
+															SNew(SImage)
+																.Image(FEditorStyle::GetBrush(TEXT("Profiler.Tab.FiltersAndPresets")))
+														]
+
+													+ SHorizontalBox::Slot()
+														.AutoWidth()
+														[
+															SNew(STextBlock)
+																.Text( LOCTEXT("FiltersAndPresetsLabel", "Filters And Presets") )
+														]
+												]
+
+											// Filters And Presets
+											+ SVerticalBox::Slot()
+												.FillHeight(1.0f)
+												.Padding(0.0f, 2.0f, 0.0f, 0.0f)
+												[
+													SAssignNew(FiltersAndPresets, SFiltersAndPresets)
+												]
+										]
 								]
 
-								+SHorizontalBox::Slot()
-								.AutoWidth()
+							+ SSplitter::Slot()
+								.Value(0.75f)
 								[
-									SNew(STextBlock)
-									.Text( LOCTEXT("GraphViewLabel", "Graph View") )
-								]
-							]
+									SNew(SSplitter)
+										.Orientation(Orient_Vertical)
+										//.PhysicalSplitterHandleSize( 2.0f )
+										//.HitDetectionSplitterHandleSize( 4.0f )
+
+										+ SSplitter::Slot()
+											.Value(0.25f)
+											[
+												SNew(SVerticalBox)
+
+												// Header
+												+SVerticalBox::Slot()
+													.AutoHeight()
+														[
+															SNew(SHorizontalBox)
+
+															+ SHorizontalBox::Slot()
+																.AutoWidth()
+																[
+																	SNew(SImage)
+																		.Image(FEditorStyle::GetBrush(TEXT("Profiler.Tab.GraphView")))
+																]
+
+															+ SHorizontalBox::Slot()
+																.AutoWidth()
+																[
+																	SNew(STextBlock)
+																		.Text(LOCTEXT("GraphViewLabel", "Graph View"))
+																]
+														]
 							
-							// Spacer
-							+SVerticalBox::Slot()
-							.AutoHeight()
-							[
-								SNew( SSpacer )
-								.Size( FVector2D( 2.0f, 2.0f ) )
+												// Graph View
+												+ SVerticalBox::Slot()
+													.FillHeight( 1.0f )
+													.Padding(0.0f, 2.0f, 0.0f, 0.0f)
+													[
+														SAssignNew(GraphPanel, SProfilerGraphPanel)
+													]
+											]
+
+										+ SSplitter::Slot()
+											.Value(0.75f)
+											[
+												SAssignNew(EventGraphPanel, SVerticalBox)
+											]
+									]
 							]
+					]
 
-							// Graph View
-							+SVerticalBox::Slot()
-							.FillHeight( 1.0f )
-							[
-								SAssignNew(GraphPanel,SProfilerGraphPanel)
-							]
-						]
-
-						+SSplitter::Slot()
-						.Value( 2.0f )
+			// session hint overlay
+			+ SOverlay::Slot()
+				.HAlign(HAlign_Center)
+				.VAlign(VAlign_Center)
+				[
+					SNew(SBorder)
+						.BorderImage(FEditorStyle::GetBrush("NotificationList.ItemBackground"))
+						.Padding(8.0f)
+						.Visibility(this, &SProfilerWindow::IsSessionOverlayVissible)
 						[
-							SAssignNew(EventGraphPanel,SVerticalBox)
+							SNew(STextBlock)
+								.Text(LOCTEXT("SelectSessionOverlayText", "Please select a session from the Session Browser or load a saved capture."))
 						]
-					]
 				]
-			]
 
-			// Overlay slot for the notification area.
-			+SOverlay::Slot()
-			.HAlign(HAlign_Center)
-			.VAlign(VAlign_Center)
-			[
-				SNew(SBorder)
-				.BorderImage( FEditorStyle::GetBrush("NotificationList.ItemBackground") )
-				.Padding( 8.0f )
-				.Visibility( this, &SProfilerWindow::IsSessionOverlayVissible )
+			// notification area overlay
+			+ SOverlay::Slot()
+				.HAlign(HAlign_Right)
+				.VAlign(VAlign_Bottom)
+				.Padding(16.0f)
 				[
-					SNew(STextBlock)
-					.Text(LOCTEXT("SelectSessionOverlayText", "Please select a session from the Session Browser or load a saved capture."))
+					SAssignNew(NotificationList, SNotificationList)
 				]
-			]
-
-			// Overlay slot for the notification area, the second
-			+SOverlay::Slot()
-			.HAlign(HAlign_Right)
-			.VAlign(VAlign_Bottom)
-			.Padding( 16.0f )
-			[
-				SAssignNew(NotificationList, SNotificationList)
-			]
 
-			// Overlay slot for the profiler settings, the third
-			+SOverlay::Slot()
-			.HAlign(HAlign_Center)
-			.VAlign(VAlign_Center)
-			.Expose( OverlaySettingsSlot )
+			// profiler settings overlay
+			+ SOverlay::Slot()
+				.HAlign(HAlign_Center)
+				.VAlign(VAlign_Center)
+				.Expose(OverlaySettingsSlot)
 		];
 
 	ProfilerMiniView->OnSelectionBoxChanged().AddSP( GraphPanel.ToSharedRef(), &SProfilerGraphPanel::MiniView_OnSelectionBoxChanged );
diff --git a/Engine/Source/Developer/Profiler/Private/Widgets/SProfilerWindow.h b/Engine/Source/Developer/Profiler/Private/Widgets/SProfilerWindow.h
--- a/Engine/Source/Developer/Profiler/Private/Widgets/SProfilerWindow.h
+++ b/Engine/Source/Developer/Profiler/Private/Widgets/SProfilerWindow.h
@@ -130,7 +130,7 @@ public:
 	TSharedPtr<SVerticalBox> EventGraphPanel;
 
 	/** Holds the filter and presets widget/slot. */
-	SHorizontalBox::FSlot* FiltersAndPresetsSlot;
+	SSplitter::FSlot* FiltersAndPresetsSlot;
 	TSharedPtr<SFiltersAndPresets> FiltersAndPresets;
 
 	/** Widget for the panel which contains all graphs and event graphs. */
@@ -153,4 +153,4 @@ public:
 
 	/** The number of seconds the profiler has been active */
 	float	DurationActive;
-};
\ No newline at end of file
+};
diff --git a/Engine/Source/Developer/SessionServices/Private/SessionManager.cpp b/Engine/Source/Developer/SessionServices/Private/SessionManager.cpp
--- a/Engine/Source/Developer/SessionServices/Private/SessionManager.cpp
+++ b/Engine/Source/Developer/SessionServices/Private/SessionManager.cpp
@@ -291,7 +291,7 @@ void FSessionManager::HandleSessionPongMessage( const FSessionServicePong& Messa
 		return;
 	}
 
-	if (Message.Standalone && !IsValidOwner(Message.SessionOwner))
+	if (!Message.Standalone && !IsValidOwner(Message.SessionOwner))
 	{
 		return;
 	}
