Profiler: Fixed profiler window not disabling when no session instances selected

diff --git a/Engine/Source/Developer/Profiler/Private/ProfilerManager.cpp b/Engine/Source/Developer/Profiler/Private/ProfilerManager.cpp
--- a/Engine/Source/Developer/Profiler/Private/ProfilerManager.cpp
+++ b/Engine/Source/Developer/Profiler/Private/ProfilerManager.cpp
@@ -404,7 +404,9 @@ void FProfilerManager::SessionManager_OnSelectedSessionChanged( const ISessionIn
 void FProfilerManager::SessionManager_OnInstanceSelectionChanged(const TSharedPtr<ISessionInstanceInfo>& /*Instance*/, bool /*Selected*/)
 {
 	const ISessionInfoPtr& SelectedSession = SessionManager->GetSelectedSession();
-	const bool SessionIsValid = SelectedSession.IsValid() && (SelectedSession->GetSessionOwner() == FPlatformProcess::UserName(true));
+	const bool SessionIsValid = SelectedSession.IsValid()
+		&& (SelectedSession->GetSessionOwner() == FPlatformProcess::UserName(true))
+		&& (SessionManager->GetSelectedInstances().Num() > 0);
 
 	if( ActiveSession != SelectedSession || FProfilerManager::GetSettings().bSingleInstanceMode )
 	{
