Copying //UE4/Dev-Core to //UE4/Main

==========================
MAJOR FEATURES + CHANGES
==========================

Change 2783106 on 2015/11/30 by [REDACTED_PERSON][REDACTED_HANDLE]

	Introduced GC UObject clusters. GC clusters provide means to create disregard for GC subsets at load time (e.g. Materials with material expressions and their textures).

	- Saves about 25ms in reachability analysis (58ms -> 33ms)
	- UObject classes/instances can now be marked as cluster root objects with CanBeClusterRoot() function override.
	- Cluster creation is automatic. Clusters don't require any manual handling for GC to collect them when nothing is referencing them.
	- Moved token stream processing to a new class FFastReferenceFinder to make it more generic and re-usable by other code
	- Removed REFERENCE_INFO macro from GC code and replaced it with a local variable (saves about ~1.9ms: 33.2ms -> 31.3ms)

Change 2773094 on 2015/11/19 by [REDACTED_PERSON][REDACTED_HANDLE]

	Multicast script delegate check for existing bindings replaced with ensure.
	Multicast native delegate no longer checks for existing bindings.
	Removal of old delegate code.
	Some FORCEINLINEing to improve debugging experience of stepping into delegate code.

Change 2782180 on 2015/11/27 by [REDACTED_PERSON][REDACTED_HANDLE]

	Make scoped seconds timer class available outside of stats build. Normal usage macros still remain guarded
	Added SCOPE_SECONDS_COUNTER_RECURSION_SAFE which only times during the outmost instance of a recursive function
	Added SCOPE_SECONDS_COUNTER_RECURSION_SAFE_BASE and SCOPE_SECONDS_COUNTER_BASE which are defined in all build types, for easy temporary timing in Test/Shipping builds.
	Added a boolean parameter to the timer class which can be used to disable it without having to mess around with scoping the calling code

Change 2782635 on 2015/11/30 by [REDACTED_PERSON][REDACTED_HANDLE]

	Added GetTimeStampPair() to the filemanager and platformfile interfaces. Requests timestamps for a pair of files where we assume that both files would always exist at the same wrapper level. Allows us to skip file system queries for localization package lookups where the native file is in a pak but the localized file doesn't exist.

Change 2775153 on 2015/11/20 by [REDACTED_PERSON][REDACTED_HANDLE]

	CrashReportServer moved out of the not for licencees, a few fixes, removed RegisterPII

Change 2775560 on 2015/11/20 by [REDACTED_PERSON][REDACTED_HANDLE]

	FDelegateBase::GetDelegateInstance deprecated and replaced with FDelegateBase::GetDelegateInstanceProtected.

Change 2781138 on 2015/11/25 by [REDACTED_PERSON][REDACTED_HANDLE]

	Stats - Converted is using a new stats reader, a few more optimizations, should be 10x times faster

Change 2772990 on 2015/11/19 by [REDACTED_PERSON][REDACTED_HANDLE]

	Fixing potential dead lock when suspending and resuming async loading multiple times

Change 2773023 on 2015/11/19 by [REDACTED_PERSON][REDACTED_HANDLE]

	Support for references added through AddReferencedObjects in FArchiveReplaceObjectRef

Change 2781055 on 2015/11/25 by [REDACTED_PERSON][REDACTED_HANDLE]

	Changes to IDelegateInstance reverted to allow licensees an easier time when upgrading their use of the now-deprecated GetDelegateInstance() code path.
	New TryGetBoundFunctionName() to aid the debugging of delegate bindings in non-shipping configs.

Change 2773114 on 2015/11/19 by [REDACTED_PERSON][REDACTED_HANDLE]

	FMath::IsPowerOfTwo is now templated to take any type.

Change 2773643 on 2015/11/19 by [REDACTED_PERSON][REDACTED_HANDLE]

	GetDelegateInstance() calls replaced - delegate instances never compare equal unless you are comparing two unbound delegates (both null) or comparing a delegate with itself.

Change 2777686 on 2015/11/23 by [REDACTED_PERSON][REDACTED_HANDLE]

	GitHub #1793 - File write flags argument

Change 2780590 on 2015/11/25 by [REDACTED_PERSON][REDACTED_HANDLE]

	Fix for FArchiveProxy::operator<< overloads.

Change 2780845 on 2015/11/25 by [REDACTED_PERSON][REDACTED_HANDLE]

Change 2780962 on 2015/11/25 by [REDACTED_PERSON][REDACTED_HANDLE]

	Stats - Added FStatsWriteStream for basic saving stat messages into a stream, initial support for reading a regular stats file, minor performance optimization, coding standard fixes

Change 2781887 on 2015/11/26 by [REDACTED_PERSON][REDACTED_HANDLE]

	Profiler/ProfilerClient - Removed unneeded synchronization points, replaces with task graph SendTo jobs, removed PROFILER_THREADED_LOAD, replaced with new stats loading mechanism,  should be around 2x times faster (4x since the optimization pass)

Change 2781893 on 2015/11/26 by [REDACTED_PERSON][REDACTED_HANDLE]

	TCachedOSPageAllocator abstracted from MallocBinned2.
	Misc tidy-ups.

Change 2782198 on 2015/11/27 by [REDACTED_PERSON][REDACTED_HANDLE]

	Profiler - Better indication of the loading progress, should no longer freeze without a progress bar

Change 2782446 on 2015/11/29 by [REDACTED_PERSON][REDACTED_HANDLE]

	Warn when calling delegates' Create* functions when they're not assigned to anything.

Change 2782538 on 2015/11/30 by [REDACTED_PERSON][REDACTED_HANDLE]

	#UE4 Removed DiskCachedAssetDataBuffer as it was not strictly necessary and was triggering a crash when loading the cached registry from disk. This data is now stored directly in DiskCachedAssetDataMap. It was already true that this map does not change outside of SerializeCache but now it is critical since NewCachedAssetDataMap keeps pointers directly to its values.

	Asset registry fixes by [REDACTED_PERSON]. Possible fix for UE-23783.

Change 2782564 on 2015/11/30 by [REDACTED_PERSON][REDACTED_HANDLE]

	FReferenceCollector::AddReferenceObjects performance improvements for TArrays. ARO will no longer call HandleObjectReference multiple times but instead will call HandleObjectReferences just once (currently only implemented for FGCCollector). Reduces the number of virtual function calls while GC'ing.

Change 2782716 on 2015/11/30 by [REDACTED_PERSON][REDACTED_HANDLE]

	UObject serial number array initialized for debug visualization.

Change 2782933 on 2015/11/30 by [REDACTED_PERSON][REDACTED_HANDLE]

	Critical sections are no longer copyable.

Change 2783061 on 2015/11/30 by [REDACTED_PERSON][REDACTED_HANDLE]

diff --git a/Engine/Config/BaseEngine.ini b/Engine/Config/BaseEngine.ini
--- a/Engine/Config/BaseEngine.ini
+++ b/Engine/Config/BaseEngine.ini
@@ -1490,6 +1490,7 @@ gc.NumRetriesBeforeForcingGC=0
 gc.AllowParallelGC=True
 gc.TimeBetweenPurgingPendingKillObjects=60
 gc.MaxObjectsInEditor=16777216
+gc.CreateGCClusters=True
 
 [Internationalization]
 +LocalizationPaths=../../../Engine/Content/Localization/Engine
diff --git a/Engine/Source/Developer/CrashDebugHelper/Private/CrashDebugHelper.cpp b/Engine/Source/Developer/CrashDebugHelper/Private/CrashDebugHelper.cpp
--- a/Engine/Source/Developer/CrashDebugHelper/Private/CrashDebugHelper.cpp
+++ b/Engine/Source/Developer/CrashDebugHelper/Private/CrashDebugHelper.cpp
@@ -1,4 +1,4 @@
-// Copyright 1998-2015 Epic Games, Inc. All Rights Reserved.
+﻿// Copyright 1998-2015 Epic Games, Inc. All Rights Reserved.
 
 #include "CrashDebugHelperPrivatePCH.h"
 #include "CrashDebugPDBCache.h"
@@ -12,7 +12,124 @@
 	#define MINIDUMPDIAGNOSTICS	0
 #endif
 
-const TCHAR* ICrashDebugHelper::P4_DEPOT_PREFIX = TEXT( "//depot/" );
+/*-----------------------------------------------------------------------------
+	FCrashDebugHelperConfig
+-----------------------------------------------------------------------------*/
+
+/**
+*  Holds FullCrashDump properties from the config.
+*
+*	PDBCache_0_Branch=UE4-Branch
+*	PDBCache_0_ExecutablePathPattern=ue4.net\Builds\UE4-Branch\%ENGINE_VERSION%
+*	PDBCache_0_SymbolPathPattern=ue4.net\Builds\UE4-Branch\%ENGINE_VERSION%
+*	
+*	If PDBCache_0_SymbolPathPattern is missing, the value from PDBCache_0_ExecutablePathPattern will be used
+*/
+struct FPDBCacheConfigEntry
+{
+	/** Initialization constructor. */
+	FPDBCacheConfigEntry( const FString& InBranch, const FString& InExecutablePathPattern, const FString& InSymbolPathPattern )
+		: Branch( InBranch )
+		, ExecutablePathPattern( InExecutablePathPattern )
+		, SymbolPathPattern( InSymbolPathPattern )
+	{}
+
+
+	/** Branch name. */
+	const FString Branch;
+
+	/** Location of the executables */
+	const FString ExecutablePathPattern;
+
+	/** Location of the symbols, usually the same as the executables. */
+	const FString SymbolPathPattern;
+};
+
+/** Helper struct for reading PDB cache configuration. */
+struct FCrashDebugHelperConfig
+{
+	static FCrashDebugHelperConfig& Get()
+	{
+		static FCrashDebugHelperConfig Instance;
+		return Instance;
+	}
+
+	bool IsValid() const
+	{
+		// We need a least one entry to proceed.
+		return PDBCacheConfigEntries.Num() > 0;
+	}
+
+	/** Reads configuration. */
+	void ReadFullCrashDumpConfigurations();
+
+	/** Gets the config for branch. */
+	FPDBCacheConfigEntry GetCacheConfigEntryForBranch( const FString& Branch ) const
+	{
+		const FString BranchFixed = Branch.Replace( TEXT( "/" ), TEXT( "+" ) );
+		for (const auto& It : PDBCacheConfigEntries)
+		{
+			if (BranchFixed.Contains( It.Branch ))
+			{
+				return It;
+			}
+		}
+
+		// Invalid entry.
+		return FPDBCacheConfigEntry( Branch, FString(), FString() );
+	}
+
+protected:
+	/** Returns empty string if couldn't read. */
+	FString GetKey( const FString& KeyName );
+
+	/** Configuration for PDB Cache. */
+	TArray<FPDBCacheConfigEntry> PDBCacheConfigEntries;
+};
+
+void FCrashDebugHelperConfig::ReadFullCrashDumpConfigurations()
+{
+	for (int32 NumEntries = 0;; ++NumEntries)
+	{
+		const FString Branch = GetKey( FString::Printf( TEXT( "PDBCache_%i_Branch" ), NumEntries ) );
+		if (Branch.IsEmpty())
+		{
+			break;
+		}
+
+		const FString ExecutablePathPattern = GetKey( FString::Printf( TEXT( "PDBCache_%i_ExecutablePathPattern" ), NumEntries ) );
+		if (ExecutablePathPattern.IsEmpty())
+		{
+			break;
+		}
+
+		FString SymbolPathPattern = GetKey( FString::Printf( TEXT( "PDBCache_%i_SymbolPathPattern" ), NumEntries ) );
+		if (SymbolPathPattern.IsEmpty())
+		{
+			SymbolPathPattern = ExecutablePathPattern;
+		}
+		
+		PDBCacheConfigEntries.Add( FPDBCacheConfigEntry( Branch, ExecutablePathPattern, SymbolPathPattern ) );
+
+		UE_LOG( LogCrashDebugHelper, Log, TEXT( "PDBCacheConfigEntry: Branch:%s ExecutablePathPattern:%s SymbolPathPattern:%s" ), *Branch, *ExecutablePathPattern, *SymbolPathPattern );
+	}
+}
+
+FString FCrashDebugHelperConfig::GetKey( const FString& KeyName )
+{
+	const FString SectionName = TEXT( "Engine.CrashDebugHelper" );
+
+	FString Result;
+	if (!GConfig->GetString( *SectionName, *KeyName, Result, GEngineIni ))
+	{
+		return TEXT( "" );
+	}
+	return Result;
+}
+
+/*-----------------------------------------------------------------------------
+	ICrashDebugHelper
+-----------------------------------------------------------------------------*/
 
 void ICrashDebugHelper::SetDepotIndex( FString& PathToChange )
 {
@@ -42,39 +159,7 @@ bool ICrashDebugHelper::Init()
 	}
 	else
 	{
-		// Look up the depot name
-		// Try to use the command line param
-		FString DepotName;
-		FString CmdLineBranchName;
-		if( FParse::Value( FCommandLine::Get(), TEXT( "BranchName=" ), CmdLineBranchName ) )
-		{
-			DepotName = FString::Printf( TEXT( "%s%s" ), P4_DEPOT_PREFIX, *CmdLineBranchName );
-		}
-		// Default to BRANCH_NAME
-		else
-		{
-			DepotName = FString::Printf( TEXT( "%s" ), *FApp::GetBranchName() );
-		}
-
-		CrashInfo.DepotName = DepotName.Replace( TEXT( "+" ), TEXT( "/" ) );
-
-		// Try to get the BuiltFromCL from command line to use this instead of attempting to locate the CL in the minidump
-		FString CmdLineBuiltFromCL;
-		int32 BuiltFromCL = -1;
-		if( FParse::Value( FCommandLine::Get(), TEXT( "BuiltFromCL=" ), CmdLineBuiltFromCL ) )
-		{
-			if( !CmdLineBuiltFromCL.IsEmpty() )
-			{
-				BuiltFromCL = FCString::Atoi( *CmdLineBuiltFromCL );
-			}
-		}
-		// Default to BUILT_FROM_CHANGELIST.
-		else
-		{
-			BuiltFromCL = FEngineVersion::Current().GetChangelist();
-		}
-
-		CrashInfo.BuiltFromCL = BuiltFromCL;
+		// Obsolete.
 	}
 
 
@@ -84,13 +169,12 @@ bool ICrashDebugHelper::Init()
 	
 	GConfig->GetString( TEXT( "Engine.CrashDebugHelper" ), TEXT( "SourceControlBuildLabelPattern" ), SourceControlBuildLabelPattern, GEngineIni );
 
-	GConfig->GetArray( TEXT( "Engine.CrashDebugHelper" ), TEXT( "ExecutablePathPattern" ), ExecutablePathPatterns, GEngineIni );
-	GConfig->GetArray( TEXT( "Engine.CrashDebugHelper" ), TEXT( "SymbolPathPattern" ), SymbolPathPatterns, GEngineIni );
-	GConfig->GetArray( TEXT( "Engine.CrashDebugHelper" ), TEXT( "Branch" ), Branches, GEngineIni );
-	const bool bCanUseSearchPatterns = Branches.Num() == ExecutablePathPatterns.Num() && ExecutablePathPatterns.Num() == SymbolPathPatterns.Num() && Branches.Num() > 0;
-	UE_CLOG( !bCanUseSearchPatterns, LogCrashDebugHelper, Warning, TEXT( "Search patterns don't match" ) );
+	FCrashDebugHelperConfig::Get().ReadFullCrashDumpConfigurations();
+
+	const bool bUsePDBCache = FCrashDebugHelperConfig::Get().IsValid();
+	UE_CLOG( !bUsePDBCache, LogCrashDebugHelper, Warning, TEXT( "CrashDebugHelperConfig invalid" ) );
 
-	if (bCanUseSearchPatterns)
+	if (bUsePDBCache)
 	{
 		FPDBCache::Get().Init();
 	}
@@ -185,8 +269,13 @@ bool ICrashDebugHelper::SyncModules()
 
 			// Find all executables.
 			TArray<FString> NetworkExecutables;
-			IFileManager::Get().FindFilesRecursive( NetworkExecutables, *CrashInfo.ExecutablesPath, TEXT( "*.dll" ), true, false, false );
-			IFileManager::Get().FindFilesRecursive( NetworkExecutables, *CrashInfo.ExecutablesPath, TEXT( "*.exe" ), true, false, false );
+
+			// Don't duplicate work.
+			if (CrashInfo.ExecutablesPath != CrashInfo.SymbolsPath)
+			{
+				IFileManager::Get().FindFilesRecursive( NetworkExecutables, *CrashInfo.ExecutablesPath, TEXT( "*.dll" ), true, false, false );
+				IFileManager::Get().FindFilesRecursive( NetworkExecutables, *CrashInfo.ExecutablesPath, TEXT( "*.exe" ), true, false, false );
+			}
 
 			// Find all symbols.
 			TArray<FString> NetworkSymbols;
@@ -215,6 +304,7 @@ bool ICrashDebugHelper::SyncModules()
 		}
 	}
 	// Get all labels associated with the crash info's label.
+	// OBSOLETE PATH
 	else if( Labels.Num() >= 1 )
 	{
 		TSharedRef<ISourceControlLabel> Label = Labels[0];
@@ -272,7 +362,7 @@ bool ICrashDebugHelper::SyncModules()
 				for( const auto& PDBPath : PDBPaths )
 				{
 					const FString PDBRelativePath = PDBPath.Replace( *CrashInfo.DepotName, TEXT( "" ) ).Replace( UESymbols, TEXT( "" ) );
-					const FString PDBFullpath = FPDBCache::Get().GetDepotRoot() / PDBPath.Replace( P4_DEPOT_PREFIX, TEXT( "" ) );
+					const FString PDBFullpath = FPDBCache::Get().GetDepotRoot() / PDBPath;
 
 					const FString PDBMatch = PDBRelativePath.Replace( TEXT( "pdb" ), TEXT( "" ) );
 					const FString NetworkRelativePath = NetworkExecutableFullpath.Replace( *CrashInfo.ExecutablesPath, TEXT( "" ) );
@@ -375,6 +465,7 @@ bool ICrashDebugHelper::SyncModules()
 			CrashInfo.PDBCacheEntry = FPDBCache::Get().CreateAndAddPDBCacheEntry( CrashInfo.LabelName, CrashInfo.DepotName, FilesToBeCached );
 		}
 	}
+	// OBSOLETE PATH
 	else
 	{
 		UE_LOG( LogCrashDebugHelper, Error, TEXT( "Could not find label: %s"), *CrashInfo.LabelName );
@@ -409,7 +500,23 @@ bool ICrashDebugHelper::ReadSourceFile( TArray<FString>& OutStrings )
 	FString FilePath;
 	if (bUsePDBCache)
 	{
-		FilePath = CrashInfo.DepotName.Replace( P4_DEPOT_PREFIX, *FPDBCache::Get().GetDepotRoot() ) / CrashInfo.SourceFile;
+		// We assume a special folder for syncing all streams and //depot/ is not used in the view mapping.
+		/*
+			//depot/... //machine/... 
+			//UE4/...	//machine/​​Stream/UE4/...​
+		*/
+		const TCHAR* DepotDelimiter = TEXT( "//depot/" );
+		const bool bIsDepot = CrashInfo.DepotName.Contains( DepotDelimiter );
+
+		if (bIsDepot)
+		{
+			FilePath = FPDBCache::Get().GetDepotRoot() / CrashInfo.DepotName.Replace( DepotDelimiter, TEXT( "" ) ) / CrashInfo.SourceFile;
+		}
+		else
+		{
+			FilePath = FPDBCache::Get().GetDepotRoot() / TEXT( "Stream" ) / CrashInfo.DepotName / CrashInfo.SourceFile;
+			FilePath.ReplaceInline( TEXT( "//" ), TEXT( "/" ) );
+		}
 	}
 	else
 	{
@@ -656,23 +763,11 @@ void ICrashDebugHelper::FindSymbolsAndBinariesStorage()
 
 	UE_LOG( LogCrashDebugHelper, Log, TEXT( "Engine version: %s" ), *CrashInfo.EngineVersion );
 
-	int32 Index = 0;
-	bool bFoundPattern = false;
-	FString ExecutablePathPattern;
-	FString SymbolPathPattern;
-	// Find branch.
-	for( ; Index < Branches.Num(); Index++ )
-	{
-		if( CrashInfo.DepotName.Contains( Branches[Index] ) )
-		{
-			bFoundPattern = true;
-			ExecutablePathPattern = ExecutablePathPatterns[Index];
-			SymbolPathPattern = SymbolPathPatterns[Index];
-			break;
-		}
-	}
+	const FPDBCacheConfigEntry ConfigEntry = FCrashDebugHelperConfig::Get().GetCacheConfigEntryForBranch( CrashInfo.DepotName );
+	FString ExecutablePathPattern = ConfigEntry.ExecutablePathPattern;
+	FString SymbolPathPattern = ConfigEntry.ExecutablePathPattern;
 
-	if( bFoundPattern )
+	if (!ExecutablePathPattern.IsEmpty() && !SymbolPathPattern.IsEmpty())
 	{
 		UE_LOG( LogCrashDebugHelper, Log, TEXT( "Using branch: %s" ), *CrashInfo.DepotName );
 	}
@@ -684,7 +779,7 @@ void ICrashDebugHelper::FindSymbolsAndBinariesStorage()
 	
 	const FString StrENGINE_VERSION = CrashInfo.EngineVersion;
 	const FString StrPLATFORM_NAME = TEXT( "" ); // Not implemented yet
-	const FString StrOLD_ENGINE_VERSION = FString::Printf( TEXT( "++depot+%s-CL-%i" ), *CrashInfo.DepotName.Replace( P4_DEPOT_PREFIX, TEXT("") ), CrashInfo.BuiltFromCL )
+	const FString StrOLD_ENGINE_VERSION = FString::Printf( TEXT( "%s-CL-%i" ), *CrashInfo.DepotName.Replace( TEXT( "+" ), TEXT( "/" ) ), CrashInfo.BuiltFromCL )
 		.Replace( TEXT("/"), TEXT("+") );
 
 	const FString TestExecutablesPath = ExecutablePathPattern
diff --git a/Engine/Source/Developer/CrashDebugHelper/Private/CrashDebugPDBCache.cpp b/Engine/Source/Developer/CrashDebugHelper/Private/CrashDebugPDBCache.cpp
--- a/Engine/Source/Developer/CrashDebugHelper/Private/CrashDebugPDBCache.cpp
+++ b/Engine/Source/Developer/CrashDebugHelper/Private/CrashDebugPDBCache.cpp
@@ -197,7 +197,7 @@ FPDBCacheEntryRef FPDBCache::CreateAndAddPDBCacheEntry( const FString& OriginalL
 	const FString EntryDirectory = PDBCachePath / CleanedLabelName;
 	const FString EntryTimeStampFilename = EntryDirectory / PDBTimeStampFile;
 
-	const FString LocalDepotDir = DepotRoot / DepotName.Replace( ICrashDebugHelper::P4_DEPOT_PREFIX, TEXT( "" ) );
+	const FString LocalDepotDir = EscapePath( DepotRoot / DepotName );
 
 	UE_LOG( LogCrashDebugHelper, Warning, TEXT( "PDB Cache entry: %s is being copied from: %s, it will take some time" ), *CleanedLabelName, *OriginalLabelName );
 	for( const auto& Filename : FilesToBeCached )
diff --git a/Engine/Source/Developer/CrashDebugHelper/Public/CrashDebugHelper.h b/Engine/Source/Developer/CrashDebugHelper/Public/CrashDebugHelper.h
--- a/Engine/Source/Developer/CrashDebugHelper/Public/CrashDebugHelper.h
+++ b/Engine/Source/Developer/CrashDebugHelper/Public/CrashDebugHelper.h
@@ -120,7 +120,7 @@ public:
 	}
 };
 
-// #YRX_Crash: 2015-07-24 Merge with FCrashDescription
+// #TODO 2015-07-24 Refactor
 /** A platform independent representation of a crash */
 class CRASHDEBUGHELPER_API FCrashInfo
 {
@@ -222,8 +222,6 @@ struct FCrashDebugInfo
 class CRASHDEBUGHELPER_API ICrashDebugHelper
 {
 public:
-	static const TCHAR* P4_DEPOT_PREFIX;
-
 	/** Replaces %DEPOT_INDEX% with the command line DepotIndex in the specified path. */
 	static void SetDepotIndex( FString& PathToChange );
 
@@ -236,16 +234,6 @@ protected:
 	 */
 	FString SourceControlBuildLabelPattern;
 	
-	/**
-	 * Patterns to search in the network driver for the executables/the symbols.
-	 * Valid from 4.2 UE builds. (CL-2068994)
-	 * This may change in future.
-	 */
-	TArray<FString> Branches;
-	TArray<FString> ExecutablePathPatterns;
-	TArray<FString> SymbolPathPatterns;
-
-	
 	/** Indicates that the crash handler is ready to do work */
 	bool bInitialized;
 
diff --git a/Engine/Source/Developer/HotReload/Private/HotReloadClassReinstancer.cpp b/Engine/Source/Developer/HotReload/Private/HotReloadClassReinstancer.cpp
--- a/Engine/Source/Developer/HotReload/Private/HotReloadClassReinstancer.cpp
+++ b/Engine/Source/Developer/HotReload/Private/HotReloadClassReinstancer.cpp
@@ -522,7 +522,9 @@ void FHotReloadClassReinstancer::ReinstanceObjectsAndUpdateDefaults()
 void FHotReloadClassReinstancer::AddReferencedObjects(FReferenceCollector& Collector)
 {
 	FBlueprintCompileReinstancer::AddReferencedObjects(Collector);
+	Collector.AllowEliminatingReferences(false);
 	Collector.AddReferencedObject(CopyOfPreviousCDO);
+	Collector.AllowEliminatingReferences(true);
 }
 
 void FHotReloadClassReinstancer::EnlistDependentBlueprintToRecompile(UBlueprint* BP, bool bBytecodeOnly)
diff --git a/Engine/Source/Developer/Profiler/Private/ProfilerCommands.cpp b/Engine/Source/Developer/Profiler/Private/ProfilerCommands.cpp
--- a/Engine/Source/Developer/Profiler/Private/ProfilerCommands.cpp
+++ b/Engine/Source/Developer/Profiler/Private/ProfilerCommands.cpp
@@ -51,14 +51,10 @@ void FProfilerCommands::RegisterCommands()
 	UI_COMMAND( ProfilerManager_ToggleLivePreview, "Live preview", "Toggles the real time live preview", EUserInterfaceActionType::ToggleButton, FInputChord() );
 
 	UI_COMMAND( DataGraph_ToggleViewMode, "Toggle graph view mode", "Toggles the data graph view mode between time based and index based", EUserInterfaceActionType::Button, FInputChord() );
-	UI_COMMAND( DataGraph_ToggleMultiMode, "Toggle graph multi mode", "Toggles the data graph multi mode between displaying area line graph and one line graph for each graph", EUserInterfaceActionType::Button, FInputChord() );
 
 	UI_COMMAND( DataGraph_ViewMode_SetTimeBased, "Time based", "Sets the data graph view mode to the time based", EUserInterfaceActionType::RadioButton, FInputChord() );
 	UI_COMMAND( DataGraph_ViewMode_SetIndexBased, "Index based", "Sets the data graph view mode to the index based", EUserInterfaceActionType::RadioButton, FInputChord() );
 
-	UI_COMMAND( DataGraph_MultiMode_SetCombined, "Combined", "Set the data graph multi mode to the displaying area line graph", EUserInterfaceActionType::RadioButton, FInputChord() );
-	UI_COMMAND( DataGraph_MultiMode_SetOneLinePerDataSource, "One line per data source", "Set the data graph multi mode to the displaying one line graph for each graph data source", EUserInterfaceActionType::RadioButton, FInputChord() );
-
 	UI_COMMAND( EventGraph_SelectAllFrames, "Select all frames", "Selects all frames in the data graph and displays them in the event graph", EUserInterfaceActionType::Button, FInputChord() );
 }
 PRAGMA_ENABLE_OPTIMIZATION
@@ -82,19 +78,6 @@ void FProfilerMenuBuilder::AddMenuEntry( FMenuBuilder& MenuBuilder, const TShare
 	);
 }
 
-/*
-
-// One session instance
-if( SessionInstanceID.IsValid() )
-{
-}
-// All session instances
-else
-{
-}
-
-}*/
-
 /*-----------------------------------------------------------------------------
 	ToggleDataPreview
 -----------------------------------------------------------------------------*/
@@ -114,74 +97,26 @@ const FUIAction FProfilerActionManager::ToggleDataPreview_Custom( const FGuid Se
 }
 
 void FProfilerActionManager::ToggleDataPreview_Execute( const FGuid SessionInstanceID )
-{
-	// One session instance
-	if( SessionInstanceID.IsValid() )
-	{
-		const FProfilerSessionRef* ProfilerSessionPtr = This->FindSessionInstance( SessionInstanceID );
-		if( ProfilerSessionPtr )
-		{
-			(*ProfilerSessionPtr)->bDataPreviewing = !(*ProfilerSessionPtr)->bDataPreviewing;
-			This->ProfilerClient->SetPreviewState( (*ProfilerSessionPtr)->bDataPreviewing, SessionInstanceID );
-		}
+{	
+	const bool bDataPreviewing = !This->IsDataPreviewing();
+	This->SetDataPreview( bDataPreviewing );
 
-		const bool bDataPreviewing = This->IsDataPreviewing();
-
-		if( !bDataPreviewing )
-		{
-			This->bLivePreview = false;
-		}
-	}
-	// All session instances
-	else
+	if (!bDataPreviewing)
 	{
-		const bool bDataPreviewing = !This->IsDataPreviewing();
-		This->SetDataPreview( bDataPreviewing );
-
-		if( !bDataPreviewing )
-		{
-			This->bLivePreview = false;
-		}
+		This->bLivePreview = false;
 	}
 }
 
 bool FProfilerActionManager::ToggleDataPreview_CanExecute( const FGuid SessionInstanceID ) const
 {
-	// One session instance
-	if( SessionInstanceID.IsValid() )
-	{
-		const bool bIsSessionInstanceValid = This->IsSessionInstanceValid( SessionInstanceID ) && This->ProfilerType == EProfilerSessionTypes::Live;
-		return bIsSessionInstanceValid;
-	}
-	// All session instances
-	else
-	{
-		const bool bCanExecute = This->ActiveSession.IsValid() && This->ProfilerType == EProfilerSessionTypes::Live && This->GetProfilerInstancesNum() > 0;
-		return bCanExecute;
-	}
+	const bool bCanExecute = This->ActiveSession.IsValid() && This->ProfilerType == EProfilerSessionTypes::Live && This->ActiveInstanceID.IsValid();
+	return bCanExecute;
 }
 
 ECheckBoxState FProfilerActionManager::ToggleDataPreview_GetCheckState( const FGuid SessionInstanceID ) const
 {
-	// One session instance
-	if( SessionInstanceID.IsValid() )
-	{
-		const FProfilerSessionRef* ProfilerSessionPtr = This->FindSessionInstance( SessionInstanceID );
-		if( ProfilerSessionPtr )
-		{
-			return (*ProfilerSessionPtr)->bDataPreviewing ? ECheckBoxState::Checked : ECheckBoxState::Unchecked;
-		}
-		else
-		{
-			return ECheckBoxState::Unchecked;
-		}
-	}
-	// All session instances
-	else
-	{
-		const bool bDataPreview = This->IsDataPreviewing();
-		return bDataPreview ? ECheckBoxState::Checked : ECheckBoxState::Unchecked;
-	}
+	const bool bDataPreview = This->IsDataPreviewing();
+	return bDataPreview ? ECheckBoxState::Checked : ECheckBoxState::Unchecked;
 }
 
 /*-----------------------------------------------------------------------------
@@ -205,7 +140,7 @@ void FProfilerActionManager::ProfilerManager_ToggleLivePreview_Execute()
 
 bool FProfilerActionManager::ProfilerManager_ToggleLivePreview_CanExecute() const
 {
-	const bool bCanExecute = This->ActiveSession.IsValid() && This->ProfilerType == EProfilerSessionTypes::Live && This->GetNumDataPreviewingInstances() > 0;
+	const bool bCanExecute = This->ActiveSession.IsValid() && This->ProfilerType == EProfilerSessionTypes::Live && This->ActiveInstanceID.IsValid();
 	return bCanExecute;
 }
 
@@ -214,83 +149,6 @@ ECheckBoxState FProfilerActionManager::ProfilerManager_ToggleLivePreview_GetChec
 	return This->bLivePreview ? ECheckBoxState::Checked : ECheckBoxState::Unchecked;
 }
 
-/*-----------------------------------------------------------------------------
-	ToggleShowDataGraph
------------------------------------------------------------------------------*/
-
-const FUIAction FProfilerActionManager::ToggleShowDataGraph_Custom( const FGuid SessionInstanceID ) const
-{
-	FUIAction UIAction;
-	UIAction.ExecuteAction = FExecuteAction::CreateRaw( this, &FProfilerActionManager::ToggleShowDataGraph_Execute, SessionInstanceID );
-	UIAction.CanExecuteAction = FCanExecuteAction::CreateRaw( this, &FProfilerActionManager::ToggleShowDataGraph_CanExecute, SessionInstanceID );
-	UIAction.GetActionCheckState = FGetActionCheckState::CreateRaw( this, &FProfilerActionManager::ToggleShowDataGraph_GetCheckState, SessionInstanceID );
-	UIAction.IsActionVisibleDelegate = FIsActionButtonVisible::CreateRaw( this, &FProfilerActionManager::ToggleShowDataGraph_IsActionButtonVisible, SessionInstanceID );
-	return UIAction;
-}
-
-void FProfilerActionManager::ToggleShowDataGraph_Execute( const FGuid SessionInstanceID )
-{
-	// One session instance
-	if( SessionInstanceID.IsValid() )
-	{
-		const bool bIsSessionInstanceTracked = This->IsSessionInstanceTracked( SessionInstanceID );
-
-		if( bIsSessionInstanceTracked )
-		{
-			This->UntrackSessionInstance( SessionInstanceID );
-		}
-		else
-		{
-			This->TrackSessionInstance( SessionInstanceID );
-		}
-	}
-}
-
-bool FProfilerActionManager::ToggleShowDataGraph_CanExecute( const FGuid SessionInstanceID ) const
-{
-	// One session instance
-	if( SessionInstanceID.IsValid() )
-	{
-		const bool bIsSessionInstanceValid = This->IsSessionInstanceValid( SessionInstanceID );
-		return bIsSessionInstanceValid;
-	}
-	// All session instances
-	else
-	{
-		return false;
-	}
-}
-
-ECheckBoxState FProfilerActionManager::ToggleShowDataGraph_GetCheckState( const FGuid SessionInstanceID ) const
-{
-	// One session instance
-	if( SessionInstanceID.IsValid() )
-	{
-		const bool bIsSessionInstanceTracked = This->IsSessionInstanceTracked( SessionInstanceID );
-		return bIsSessionInstanceTracked ? ECheckBoxState::Checked : ECheckBoxState::Unchecked;
-	}
-	// All session instances
-	else
-	{
-		return ECheckBoxState::Unchecked;
-	}
-}
-
-bool FProfilerActionManager::ToggleShowDataGraph_IsActionButtonVisible( const FGuid SessionInstanceID ) const
-{
-	// One session instance
-	if( SessionInstanceID.IsValid() )
-	{
-		return true;
-	}
-	// All session instances
-	else
-	{
-		// Hiding all graphs for all session instances as it doesn't add any useful functionality
-		return false;
-	}
-}
-
 /*-----------------------------------------------------------------------------
 	ProfilerManager_Load
 -----------------------------------------------------------------------------*/
@@ -321,7 +179,6 @@ void FProfilerActionManager::ProfilerManager_Load_Execute()
 			ProfilingDirectory, 
 			TEXT(""), 
 			LOCTEXT("ProfilerManager_Load_FileFilter", "Stats files (*.ue4stats)|*.ue4stats|Raw Stats files (*.ue4statsraw)|*.ue4statsraw").ToString(), 
-			//FProfilerManager::GetSettings().bSingleInstanceMode ? EFileDialogFlags::None : EFileDialogFlags::Multiple, 
 			EFileDialogFlags::None,
 			OutFiles
 		);
@@ -347,7 +204,7 @@ void FProfilerActionManager::ProfilerManager_Load_Execute()
 
 bool FProfilerActionManager::ProfilerManager_Load_CanExecute() const
 {
-	const bool bIsConnectionActive = This->GetNumDataCapturingInstances() > 0 || This->GetNumDataPreviewingInstances() > 0 || This->bLivePreview;
+	const bool bIsConnectionActive = This->IsDataCapturing() || This->IsDataPreviewing() || This->IsLivePreview();
 	return !(bIsConnectionActive && This->ProfilerType == EProfilerSessionTypes::Live);
 }
 
@@ -371,26 +228,12 @@ const FUIAction FProfilerActionManager::ToggleDataCapture_Custom( const FGuid Se
 
 void FProfilerActionManager::ToggleDataCapture_Execute( const FGuid SessionInstanceID )
 {
-	// One session instance
-	if( SessionInstanceID.IsValid() )
-	{
-		const FProfilerSessionRef* ProfilerSessionPtr = This->FindSessionInstance( SessionInstanceID );
-		if( ProfilerSessionPtr )
-		{
-			(*ProfilerSessionPtr)->bDataCapturing = !(*ProfilerSessionPtr)->bDataCapturing;
-			This->ProfilerClient->SetCaptureState( (*ProfilerSessionPtr)->bDataCapturing, SessionInstanceID );
-		}	
-	}
-	// All session instances
-	else
-	{
-		const bool bDataCapturing = This->IsDataCapturing();
-		This->SetDataCapture( !bDataCapturing );
-	}
+	const bool bDataCapturing = This->IsDataCapturing();
+	This->SetDataCapture( !bDataCapturing );
 
 	// Assumes that when data capturing is off, we have captured stats files on the service side.
-	const bool bDataCapturing = This->IsDataCapturing();
-	if( !bDataCapturing )
+	const bool bNewDataCapturing = This->IsDataCapturing();
+	if (!bNewDataCapturing)
 	{
 		EAppReturnType::Type Result = FPlatformMisc::MessageBoxExt
 		( 
@@ -408,45 +251,14 @@ void FProfilerActionManager::ToggleDataCapture_Execute( const FGuid SessionInsta
 
 bool FProfilerActionManager::ToggleDataCapture_CanExecute( const FGuid SessionInstanceID ) const
 {
-	// One session instance
-	if( SessionInstanceID.IsValid() )
-	{
-		const bool bIsSessionInstanceValid = This->IsSessionInstanceValid( SessionInstanceID ) && This->ProfilerType == EProfilerSessionTypes::Live;
-		return bIsSessionInstanceValid;
-	}
-	// All session instances
-	else
-	{
-		const bool bCanExecute = This->ActiveSession.IsValid() && This->ProfilerType == EProfilerSessionTypes::Live && This->GetProfilerInstancesNum() > 0;
-		return bCanExecute;
-	}
-
-	return false;
+	const bool bCanExecute = This->ActiveSession.IsValid() && This->ProfilerType == EProfilerSessionTypes::Live && This->ActiveInstanceID.IsValid();
+	return bCanExecute;
 }
 
 ECheckBoxState FProfilerActionManager::ToggleDataCapture_GetCheckState( const FGuid SessionInstanceID ) const
 {
-	// One session instance
-	if( SessionInstanceID.IsValid() )
-	{
-		const FProfilerSessionRef* ProfilerSessionPtr = This->FindSessionInstance( SessionInstanceID );
-		if( ProfilerSessionPtr )
-		{
-			return (*ProfilerSessionPtr)->bDataCapturing ? ECheckBoxState::Checked : ECheckBoxState::Unchecked;
-		}
-		else
-		{
-			return ECheckBoxState::Unchecked;
-		}
-	}
-	// All session instances
-	else
-	{
-		const bool bDataCapturing = This->IsDataCapturing();
-		return bDataCapturing ? ECheckBoxState::Checked : ECheckBoxState::Unchecked;
-	}
-
-	return ECheckBoxState::Unchecked;
+	const bool bDataCapturing = This->IsDataCapturing();
+	return bDataCapturing ? ECheckBoxState::Checked : ECheckBoxState::Unchecked;
 }
 
 /*-----------------------------------------------------------------------------
diff --git a/Engine/Source/Developer/Profiler/Private/ProfilerCommands.h b/Engine/Source/Developer/Profiler/Private/ProfilerCommands.h
--- a/Engine/Source/Developer/Profiler/Private/ProfilerCommands.h
+++ b/Engine/Source/Developer/Profiler/Private/ProfilerCommands.h
@@ -66,26 +66,15 @@ public:
 	/** Toggles the real time live preview. Global version. */
 	TSharedPtr< FUICommandInfo > ProfilerManager_ToggleLivePreview;
 
-
 	/** Toggles the data graph view mode between time based and index based. */
 	TSharedPtr< FUICommandInfo > DataGraph_ToggleViewMode;
 
-	/** Toggles the data graph multi mode between displaying area line graph for each graph data source. */
-	TSharedPtr< FUICommandInfo > DataGraph_ToggleMultiMode;
-
-
 	/** Sets the data graph view mode to the time based. */
 	TSharedPtr< FUICommandInfo > DataGraph_ViewMode_SetTimeBased;
 
 	/** Sets the data graph view mode to the index based. */
 	TSharedPtr< FUICommandInfo > DataGraph_ViewMode_SetIndexBased;
 
-	/** Set the data graph multi mode to the displaying area line graph. */
-	TSharedPtr< FUICommandInfo > DataGraph_MultiMode_SetCombined;
-
-	/** Set the data graph multi mode to the displaying one line graph for each graph data source. */
-	TSharedPtr< FUICommandInfo > DataGraph_MultiMode_SetOneLinePerDataSource;
-
 	/** Select all frame in the data graph and display them in the event graph, technically switches to the begin of history. */
 	TSharedPtr< FUICommandInfo > EventGraph_SelectAllFrames;
 };
@@ -193,29 +182,6 @@ protected:
 	/** Handles FGetActionCheckState for ToggleDataCapture. */
 	ECheckBoxState ToggleDataCapture_GetCheckState( const FGuid SessionInstanceID ) const;
 
-	/*-----------------------------------------------------------------------------
-		ToggleShowDataGraph
-	-----------------------------------------------------------------------------*/
-public:
-
-	/**
-	 * UI action that toggles showing data graph for the specified session instance.
-	 *
-	 * @param SessionInstanceID - the session instance that this action will be executed on, if not valid, all session instances will be used
-	 *
-	 */
-	const FUIAction ToggleShowDataGraph_Custom( const FGuid SessionInstanceID ) const;
-
-protected:
-	/** Handles FExecuteAction for ToggleShowDataGraph_Execute. */
-	void ToggleShowDataGraph_Execute( const FGuid SessionInstanceID );
-	/** Handles FCanExecuteAction for ToggleShowDataGraph_Execute. */
-	bool ToggleShowDataGraph_CanExecute( const FGuid SessionInstanceID ) const;
-	/** Handles FGetActionCheckState for ToggleShowDataGraph_Execute. */
-	ECheckBoxState ToggleShowDataGraph_GetCheckState( const FGuid SessionInstanceID ) const;
-	/** Handles IsActionButtonVisible for ToggleShowDataGraph_Execute. */
-	bool ToggleShowDataGraph_IsActionButtonVisible( const FGuid SessionInstanceID ) const;
-
 private:
 	/** Private constructor. */
 	FProfilerActionManager( class FProfilerManager* Instance )
diff --git a/Engine/Source/Developer/Profiler/Private/ProfilerDataSource.cpp b/Engine/Source/Developer/Profiler/Private/ProfilerDataSource.cpp
--- a/Engine/Source/Developer/Profiler/Private/ProfilerDataSource.cpp
+++ b/Engine/Source/Developer/Profiler/Private/ProfilerDataSource.cpp
@@ -399,7 +399,8 @@ FEventGraphData::FEventGraphData( const FProfilerSessionRef& InProfilerSession,
 	Description = FString::Printf( TEXT("%s: %i"), *InProfilerSession->GetShortName(), InFrameIndex );
 
 	// @TODO: Duplicate is not needed, remove it later.
-	const IDataProviderRef DataProvider = InProfilerSession->GetDataProvider()->Duplicate<FArrayDataProvider>( FrameStartIndex );
+	const IDataProviderRef& SessionDataProvider = InProfilerSession->GetDataProvider(); 
+	const IDataProviderRef DataProvider = SessionDataProvider->Duplicate<FArrayDataProvider>( FrameStartIndex );
 
 	const double FrameDurationMS = DataProvider->GetFrameTimeMS( 0 ); 
 	const FProfilerSample& RootProfilerSample = DataProvider->GetCollection()[0];
@@ -459,7 +460,7 @@ FEventGraphData::FEventGraphData( const FEventGraphData& Source )
 	RootEvent = Source.GetRoot()->DuplicateWithHierarchyPtr();
 	RootEvent->SetRootAndThreadForAllChildren();
 	Advance( Source.FrameStartIndex, Source.FrameEndIndex );
-	}
+}
 
 FEventGraphDataRef FEventGraphData::DuplicateAsRef()
 {
diff --git a/Engine/Source/Developer/Profiler/Private/ProfilerDataSource.h b/Engine/Source/Developer/Profiler/Private/ProfilerDataSource.h
--- a/Engine/Source/Developer/Profiler/Private/ProfilerDataSource.h
+++ b/Engine/Source/Developer/Profiler/Private/ProfilerDataSource.h
@@ -2,6 +2,37 @@
 
 #pragma once
 
+/*-----------------------------------------------------------------------------
+Event graph related type definitions
+-----------------------------------------------------------------------------*/
+
+/** Type definition for shared pointers to instances of FEventGraphSample. */
+typedef TSharedPtr<class FEventGraphSample> FEventGraphSamplePtr;
+
+/** Type definition for shared references to instances of FEventGraphSample. */
+typedef TSharedRef<class FEventGraphSample> FEventGraphSampleRef;
+
+/** Type definition for weak references to instances of FEventGraphSample. */
+typedef TWeakPtr<class FEventGraphSample> FEventGraphSampleWeak;
+
+/** Type definition for shared pointers to instances of FEventGraphData. */
+typedef TSharedPtr<class FEventGraphData, ESPMode::ThreadSafe> FEventGraphDataPtr;
+
+/** Type definition for shared references to instances of FEventGraphData. */
+typedef TSharedRef<class FEventGraphData, ESPMode::ThreadSafe> FEventGraphDataRef;
+
+/** Type definition for shared pointers to instances of FEventGraphDataHandler. */
+typedef TSharedPtr<class FEventGraphDataHandler> FEventGraphDataHandlerPtr;
+
+/** Type definition for shared references to instances of FEventGraphDataHandler. */
+typedef TSharedRef<class FEventGraphDataHandler> FEventGraphDataHandlerRef;
+
+/** Scratch buffers for multithreaded usage. */
+struct FProfilerScratchArea : public TThreadSingleton<FProfilerScratchArea>
+{
+	TArray<FEventGraphSample*> ExecuteOperationArray;
+};
+
 /*-----------------------------------------------------------------------------
 	TimeAccuracy
 -----------------------------------------------------------------------------*/
@@ -672,32 +703,6 @@ protected:
 	TMap<FGuid,FGraphDataSourceRefConst> GraphDataSources;
 };
 
-/*-----------------------------------------------------------------------------
-	Event graph related type definitions
------------------------------------------------------------------------------*/
-
-/** Type definition for shared pointers to instances of FEventGraphSample. */
-typedef TSharedPtr<class FEventGraphSample> FEventGraphSamplePtr;
-
-/** Type definition for shared references to instances of FEventGraphSample. */
-typedef TSharedRef<class FEventGraphSample> FEventGraphSampleRef;
-
-/** Type definition for weak references to instances of FEventGraphSample. */
-typedef TWeakPtr<class FEventGraphSample> FEventGraphSampleWeak;
-
-/** Type definition for shared pointers to instances of FEventGraphData. */
-typedef TSharedPtr<class FEventGraphData, ESPMode::ThreadSafe> FEventGraphDataPtr;
-
-/** Type definition for shared references to instances of FEventGraphData. */
-typedef TSharedRef<class FEventGraphData, ESPMode::ThreadSafe> FEventGraphDataRef;
-
-/** Type definition for shared pointers to instances of FEventGraphDataHandler. */
-typedef TSharedPtr<class FEventGraphDataHandler> FEventGraphDataHandlerPtr;
-
-/** Type definition for shared references to instances of FEventGraphDataHandler. */
-typedef TSharedRef<class FEventGraphDataHandler> FEventGraphDataHandlerRef;
-
-
 /*-----------------------------------------------------------------------------
 	Minimal event graph sample property management
 -----------------------------------------------------------------------------*/
@@ -964,12 +969,6 @@ struct FEventGraphConsts
  */
 class FEventGraphSample : public TSharedFromThis<FEventGraphSample>
 {
-	enum
-	{
-		/** Maximum number of stack when traversing event graph - should be enough to store root nodes and all its children.*/
-		MaxStackSize = 256*1024
-	};
-
 	struct FDuplicateHierarchyTag {};
 	struct FDuplicateSimpleTag {};
 
@@ -1180,7 +1179,7 @@ public:
 		_ExclusiveTimeMS = FMath::Max( _ExclusiveTimeMS, Other._ExclusiveTimeMS );
 	}
 
-	const bool AreTheSamePtr( const FEventGraphSamplePtr& Other ) const
+	FORCEINLINE_DEBUGGABLE  bool AreTheSamePtr( const FEventGraphSamplePtr& Other ) const
 	{
 		return *this == *Other.Get();
 	}
@@ -1230,16 +1229,18 @@ public:
 	}
 
 	template< typename TFunc > 
-	FORCEINLINE void ExecuteOperationForAllChildren( TFunc FuncToCall )
+	FORCEINLINE_DEBUGGABLE void ExecuteOperationForAllChildren( TFunc FuncToCall )
 	{
-		FEventGraphSample** Stack = new FEventGraphSample*[MaxStackSize];
-		Stack[ 0 ] = &AsShared().Get();
+		const bool bAllowShrinking = false;
+		TArray<FEventGraphSample*>& Stack = FProfilerScratchArea::Get().ExecuteOperationArray;
+
+		Stack.Add( &AsShared().Get() );
 		int32 Idx = 1;
 
-		while( Idx > 0 )
+		while (Stack.Num() > 0)
 		{
 			// Get the parent and assign events.
-			FEventGraphSample* Current = Stack[ --Idx ];
+			FEventGraphSample* Current = Stack.Pop( bAllowShrinking );
 			FuncToCall( Current );
 
 			// Push children onto the stack.
@@ -1248,39 +1249,41 @@ public:
 			for( int32 ChildIndex = 0; ChildIndex < NumChildren; ++ChildIndex )
 			{
 				const FEventGraphSamplePtr& ChildPtr = ChildrenPtr[ChildIndex];
-				Stack[ Idx++ ] = ChildPtr.Get();
+				Stack.Push( ChildPtr.Get() );
 			}
-			check(Idx < MaxStackSize);
 		}
 
-		delete [] Stack;
+		// Reset, but keep the allocation.
+		Stack.Reset();
 	}
 
 	template< typename TFunc, typename TArg0 > 
-	FORCEINLINE void ExecuteOperationForAllChildren( TFunc FuncToCall, TArg0& Arg0 )
+	FORCEINLINE_DEBUGGABLE void ExecuteOperationForAllChildren( TFunc FuncToCall, TArg0& Arg0 )
 	{
-		FEventGraphSample** Stack = new FEventGraphSample*[MaxStackSize]; 
-		Stack[ 0 ] = &AsShared().Get();
+		const bool bAllowShrinking = false;
+		TArray<FEventGraphSample*>& Stack = FProfilerScratchArea::Get().ExecuteOperationArray;
+
+		Stack.Add( &AsShared().Get() );
 		int32 Idx = 1;
 
-		while( Idx > 0 )
+		while (Stack.Num() > 0)
 		{
 			// Get the parent and assign events.
-			FEventGraphSample* Current = Stack[ --Idx ];
+			FEventGraphSample* Current = Stack.Pop( bAllowShrinking );
 			FuncToCall( Current, Arg0 );
 
 			// Push children onto the stack.
 			const TArray<FEventGraphSamplePtr>& ChildrenPtr = Current->GetChildren();
 			const int32 NumChildren = ChildrenPtr.Num();
-			for( int32 ChildIndex = 0; ChildIndex < NumChildren; ++ChildIndex )
+			for (int32 ChildIndex = 0; ChildIndex < NumChildren; ++ChildIndex)
 			{
 				const FEventGraphSamplePtr& ChildPtr = ChildrenPtr[ChildIndex];
-				Stack[ Idx++ ] = ChildPtr.Get();
+				Stack.Push( ChildPtr.Get() );
 			}
-			check(Idx < MaxStackSize);
 		}
 
-		delete [] Stack;
+		// Reset, but keep the allocation.
+		Stack.Reset();
 	}
 
 protected:
@@ -1327,42 +1330,32 @@ protected:
 		_ThreadPtr = ThreadEvent;
 	}
 
-	void SetRootAndThreadEvents_Iterative( FEventGraphSample* RootEvent, FEventGraphSample* ThreadEvent )
+	FORCEINLINE_DEBUGGABLE void SetRootAndThreadEvents_Iterative( FEventGraphSample* RootEvent, FEventGraphSample* ThreadEvent )
 	{
-		FEventGraphSample** Stack = new FEventGraphSample*[MaxStackSize];
-		Stack[ 0 ] = &AsShared().Get();
+		const bool bAllowShrinking = false;
+		TArray<FEventGraphSample*>& Stack = FProfilerScratchArea::Get().ExecuteOperationArray;
+
+		Stack.Add( &AsShared().Get() );
 		int32 Idx = 1;
 
-		while( Idx > 0 )
+		while (Stack.Num() > 0)
 		{
 			// Get the parent and assign events.
-			FEventGraphSample* Current = Stack[ --Idx ];
+			FEventGraphSample* Current = Stack.Pop( bAllowShrinking );
 			Current->SetRootAndThread( RootEvent, ThreadEvent );
 
 			// Push children onto the stack.
 			const TArray<FEventGraphSamplePtr>& ChildrenPtr = Current->GetChildren();
 			const int32 NumChildren = ChildrenPtr.Num();
-			for( int32 ChildIndex = 0; ChildIndex < NumChildren; ++ChildIndex )
+			for (int32 ChildIndex = 0; ChildIndex < NumChildren; ++ChildIndex)
 			{
 				const FEventGraphSamplePtr& ChildPtr = ChildrenPtr[ChildIndex];
-				Stack[ Idx++ ] = ChildPtr.Get();
+				Stack.Push( ChildPtr.Get() );
 			}
-			check(Idx < MaxStackSize);
 		}
 
-		delete [] Stack;
-	}
-
-	/** Not used, optimized version @see SetRootAndThreadEvents_Iterative. */
-	void SetRootAndThreadEvents_Recurrent( FEventGraphSample* RootEvent, FEventGraphSample* ThreadEvent )
-	{
-		SetRootAndThread( RootEvent, ThreadEvent );
-		const int32 NumChildren = _ChildrenPtr.Num();
-		for( int32 ChildIndex = 0; ChildIndex < NumChildren; ++ChildIndex )
-		{
-			const FEventGraphSamplePtr& ChildPtr = _ChildrenPtr[ChildIndex];
-			ChildPtr->SetRootAndThreadEvents_Recurrent( RootEvent, ThreadEvent );
-		}
+		// Reset, but keep the allocation.
+		Stack.Reset();
 	}
 
 public:
diff --git a/Engine/Source/Developer/Profiler/Private/ProfilerManager.cpp b/Engine/Source/Developer/Profiler/Private/ProfilerManager.cpp
--- a/Engine/Source/Developer/Profiler/Private/ProfilerManager.cpp
+++ b/Engine/Source/Developer/Profiler/Private/ProfilerManager.cpp
@@ -110,14 +110,12 @@ void FProfilerManager::PostConstructor()
 	ProfilerClient->OnProfilerClientDisconnected().AddSP(this, &FProfilerManager::ProfilerClient_OnClientDisconnected);
 	
 	ProfilerClient->OnLoadStarted().AddSP(this, &FProfilerManager::ProfilerClient_OnLoadStarted);
-	ProfilerClient->OnLoadedMetaData().AddSP(this, &FProfilerManager::ProfilerClient_OnLoadedMetaData);
 	ProfilerClient->OnLoadCompleted().AddSP(this, &FProfilerManager::ProfilerClient_OnLoadCompleted);
 
 	ProfilerClient->OnMetaDataUpdated().AddSP(this, &FProfilerManager::ProfilerClient_OnMetaDataUpdated);
 
 	ProfilerClient->OnProfilerFileTransfer().AddSP(this, &FProfilerManager::ProfilerClient_OnProfilerFileTransfer);
 
-	SessionManager->OnCanSelectSession().AddSP( this, &FProfilerManager::SessionManager_OnCanSelectSession );
 	SessionManager->OnSelectedSessionChanged().AddSP( this, &FProfilerManager::SessionManager_OnSelectedSessionChanged );
 	SessionManager->OnInstanceSelectionChanged().AddSP( this, &FProfilerManager::SessionManager_OnInstanceSelectionChanged );
 
@@ -144,6 +142,19 @@ FProfilerManager::~FProfilerManager()
 	// Unregister tick function.
 	FTicker::GetCoreTicker().RemoveTicker( OnTickHandle );
 
+	// Remove ourselves from the session manager.
+	if (SessionManager.IsValid())
+	{
+		// clear the selected session
+		ISessionInfoPtr Ptr;
+		Ptr.Reset();
+		SessionManager->SelectSession( Ptr );
+
+		SessionManager->OnCanSelectSession().RemoveAll( this );
+		SessionManager->OnSelectedSessionChanged().RemoveAll( this );
+		SessionManager->OnInstanceSelectionChanged().RemoveAll( this );
+	}
+
 	// Remove ourselves from the profiler client.
 	if( ProfilerClient.IsValid() )
 	{
@@ -153,66 +164,67 @@ FProfilerManager::~FProfilerManager()
 		ProfilerClient->OnProfilerClientConnected().RemoveAll(this);
 		ProfilerClient->OnProfilerClientDisconnected().RemoveAll(this);
 		ProfilerClient->OnMetaDataUpdated().RemoveAll(this);
-		ProfilerClient->OnLoadedMetaData().RemoveAll(this);
 		ProfilerClient->OnLoadCompleted().RemoveAll(this);
 		ProfilerClient->OnLoadStarted().RemoveAll(this);
 		ProfilerClient->OnProfilerFileTransfer().RemoveAll(this);
 	}
 
-	// Remove ourselves from the session manager.
-	if( SessionManager.IsValid() )
-	{
-		SessionManager->OnCanSelectSession().RemoveAll(this);
-		SessionManager->OnSelectedSessionChanged().RemoveAll(this);
-		SessionManager->OnInstanceSelectionChanged().RemoveAll(this);
-
-		// clear the selected session
-		ISessionInfoPtr Ptr;
-		Ptr.Reset();
-		SessionManager->SelectSession(Ptr);
-	}
-
 	ClearStatsAndInstances();
 }
 
+class FStatsHeaderReader : public FStatsReadFile
+{
+	friend struct FStatsReader<FStatsHeaderReader>;
+	typedef FStatsReadFile Super;
+
+protected:
 
-void FProfilerManager::LoadProfilerCapture( const FString& ProfilerCaptureFilepath, const bool bAdd /*= false*/ )
+	/** Initialization constructor. */
+	FStatsHeaderReader( const TCHAR* InFilename )
+		: FStatsReadFile( InFilename, false )
+	{}
+};
+
+static int32 GetNumFrameFromCaptureSlow( const FString& ProfilerCaptureFilepath )
 {
-	// deselect the active session
-	if (ActiveSession.IsValid())
-	{
-		SessionManager->SelectSession(NULL);
-	}
+	TAutoPtr<FStatsHeaderReader> Instance( FStatsReader<FStatsHeaderReader>::Create( *ProfilerCaptureFilepath ) );
+	return Instance->GetNumFrames();
+}
 
-	if( bAdd == false )
+void FProfilerManager::LoadProfilerCapture( const FString& ProfilerCaptureFilepath )
+{
+	// Deselect the active session.
+	if (ActiveSession.IsValid())
 	{
-		ClearStatsAndInstances();
+		SessionManager->SelectSession( nullptr );
 	}
 
-	FProfilerSessionRef ProfilerSession = MakeShareable( new FProfilerSession( ProfilerCaptureFilepath ) );
-	auto Var = ProfilerSession->AsShared();
-	const FGuid ProfilerInstanceID = ProfilerSession->GetInstanceID();
+	ClearStatsAndInstances();
+	
+	ProfilerSession = MakeShareable( new FProfilerSession( ProfilerCaptureFilepath ) );
+	ActiveInstanceID = ProfilerSession->GetInstanceID();
 
 	ProfilerSession->
 		SetOnCaptureFileProcessed( FProfilerSession::FCaptureFileProcessedDelegate::CreateSP( this, &FProfilerManager::ProfilerSession_OnCaptureFileProcessed ) )
 		.SetOnAddThreadTime( FProfilerSession::FAddThreadTimeDelegate::CreateSP( this, &FProfilerManager::ProfilerSession_OnAddThreadTime ) );
 
-	ProfilerSessionInstances.Add( ProfilerInstanceID, ProfilerSession );
-	{
-		PROFILER_SCOPE_LOG_TIME( TEXT( "ProfilerClient->LoadCapture" ), nullptr );
-		ProfilerClient->LoadCapture( ProfilerCaptureFilepath, ProfilerInstanceID );
-	}
+	ProfilerClient->LoadCapture( ProfilerCaptureFilepath, ActiveInstanceID );
+
+	const int32 NumFrames = GetNumFrameFromCaptureSlow( ProfilerCaptureFilepath );
+	ProfilerSession->SetNumberOfFrames( NumFrames );
 
 	SessionInstancesUpdatedEvent.Broadcast();
 	ProfilerType = EProfilerSessionTypes::StatsFile;
 	
-	GetProfilerWindow()->ManageEventGraphTab( ProfilerInstanceID, true, ProfilerSession->GetName() );
+	GetProfilerWindow()->ManageEventGraphTab( ActiveInstanceID, true, ProfilerSession->GetName() );
 	SetViewMode( EProfilerViewMode::LineIndexBased );
 }
 
 void FProfilerManager::LoadRawStatsFile( const FString& RawStatsFileFileath )
 {
-	if( ActiveSession.IsValid() )
+	// #YRX_Rework 2015-11-30 
+#if	0
+	if (ActiveSession.IsValid())
 	{
 		SessionManager->SelectSession( NULL );
 	}
@@ -220,13 +232,13 @@ void FProfilerManager::LoadRawStatsFile( const FString& RawStatsFileFileath )
 
 	TSharedRef<FRawProfilerSession> ProfilerSession = MakeShareable( new FRawProfilerSession( RawStatsFileFileath ) );
 	const FGuid ProfilerInstanceID = ProfilerSession->GetInstanceID();
-	ProfilerSessionInstances.Add( ProfilerInstanceID, ProfilerSession );
+	//ProfilerSessionInstances.Add( ProfilerInstanceID, ProfilerSession );
 
 	ProfilerSession->
 		//SetOnCaptureFileProcessed( FProfilerSession::FCaptureFileProcessedDelegate::CreateSP( this, /&FProfilerManager::ProfilerSession_OnCaptureFileProcessed ) )
 		SetOnAddThreadTime( FProfilerSession::FAddThreadTimeDelegate::CreateSP( this, &FProfilerManager::ProfilerSession_OnAddThreadTime ) );
 
-	ProfilerSessionInstances.Add( ProfilerInstanceID, ProfilerSession );
+	//ProfilerSessionInstances.Add( ProfilerInstanceID, ProfilerSession );
 
 	ProfilerSession->PrepareLoading();
 	RequestFilterAndPresetsUpdateEvent.Broadcast();
@@ -236,11 +248,12 @@ void FProfilerManager::LoadRawStatsFile( const FString& RawStatsFileFileath )
 
 	SessionInstancesUpdatedEvent.Broadcast();
 	ProfilerType = EProfilerSessionTypes::StatsFileRaw;
-	
+
 	GetProfilerWindow()->ManageEventGraphTab( ProfilerInstanceID, true, ProfilerSession->GetName() );
 	SetViewMode( EProfilerViewMode::ThreadViewTimeBased );
 
 	GetProfilerWindow()->GraphPanel->ThreadView->AttachProfilerStream( ProfilerSession->GetStream() );
+#endif // 0
 }
 
 
@@ -261,26 +274,54 @@ bool FProfilerManager::Tick( float DeltaTime )
 	return true;
 }
 
+double LoadStartTime = 0.0;
+
 void FProfilerManager::ProfilerSession_OnCaptureFileProcessed( const FGuid ProfilerInstanceID )
 {
-	const FProfilerSessionRef* ProfilerSession = FindSessionInstance( ProfilerInstanceID );
-	if( ProfilerSession && ProfilerWindow.IsValid())
+	if (ProfilerSession.IsValid() && ProfilerWindow.IsValid())
 	{
 		TrackDefaultStats();
 
 		RequestFilterAndPresetsUpdateEvent.Broadcast();
 
-		GetProfilerWindow()->UpdateEventGraph( ProfilerInstanceID, (*ProfilerSession)->GetEventGraphDataAverage(), (*ProfilerSession)->GetEventGraphDataMaximum(), true );
+		GetProfilerWindow()->UpdateEventGraph( ProfilerInstanceID, ProfilerSession->GetEventGraphDataAverage(), ProfilerSession->GetEventGraphDataMaximum(), true );
 		bHasCaptureFileFullyProcessed = true;
+
+		const double TotalLoadTime = FPlatformTime::Seconds() - LoadStartTime;
+
+		const FString Description = ProfilerSession->GetName();
+		UE_LOG( LogStats, Warning, TEXT( "OnCaptureFileProcessed: %s" ), *Description );
+		UE_LOG( LogStats, Warning, TEXT( "TotalLoadTime: %.2f" ), TotalLoadTime );
+
+		// Update the notification that a file has been fully processed.
+		GetProfilerWindow()->ManageLoadingProgressNotificationState( Description, EProfilerNotificationTypes::LoadingOfflineCapture, ELoadingProgressStates::Loaded, 1.0f );
 	}
 }
 
+const bool FProfilerManager::IsDataPreviewing() const
+{
+	return IsConnected() && ProfilerSession.IsValid() && ProfilerSession->bDataPreviewing;
+}
+
+void FProfilerManager::SetDataPreview( const bool bRequestedDataPreviewState )
+{
+	ProfilerClient->SetPreviewState( bRequestedDataPreviewState );
+	if (ProfilerSession.IsValid())
+	{
+		ProfilerSession->bDataPreviewing = bRequestedDataPreviewState;
+	}
+}
+
+const bool FProfilerManager::IsDataCapturing() const
+{
+	return IsConnected() && ProfilerSession.IsValid() && ProfilerSession->bDataCapturing;
+}
+
 void FProfilerManager::SetDataCapture( const bool bRequestedDataCaptureState )
 {
 	ProfilerClient->SetCaptureState( bRequestedDataCaptureState );
-	for( auto It = GetProfilerInstancesIterator(); It; ++It )
+	if (ProfilerSession.IsValid())
 	{
-		FProfilerSessionRef ProfilerSession = It.Value();
 		ProfilerSession->bDataCapturing = bRequestedDataCaptureState;
 	}
 }
@@ -289,63 +330,55 @@ void FProfilerManager::SetDataCapture( const bool bRequestedDataCaptureState )
 	ProfilerClient
 -----------------------------------------------------------------------------*/
 
-void FProfilerManager::ProfilerClient_OnProfilerData( const FGuid& InstanceID, const FProfilerDataFrame& Content, const float DataLoadingProgress )
+void FProfilerManager::ProfilerClient_OnProfilerData( const FGuid& InstanceID, const FProfilerDataFrame& Content )
 {
 	SCOPE_CYCLE_COUNTER(STAT_PM_HandleProfilerData);
 
-	const FProfilerSessionRef* ProfilerSession = FindSessionInstance( InstanceID );
-	if( ProfilerSession && ProfilerWindow.IsValid())
+	if (ProfilerSession.IsValid())
 	{
-		(*ProfilerSession)->UpdateProfilerData( Content );
+		ProfilerSession->UpdateProfilerData( Content );
 		// Game thread should always be enabled.
 		TrackDefaultStats();
-
-		// Update the notification that a file is being loaded.
-		GetProfilerWindow()->ManageLoadingProgressNotificationState( (*ProfilerSession)->GetName(), EProfilerNotificationTypes::LoadingOfflineCapture, ELoadingProgressStates::InProgress, DataLoadingProgress );
 	}
 }
 
-void FProfilerManager::ProfilerClient_OnMetaDataUpdated( const FGuid& InstanceID )
+void FProfilerManager::ProfilerClient_OnMetaDataUpdated( const FGuid& InstanceID, const FStatMetaData& MetaData )
 {
-	const FProfilerSessionRef* ProfilerSession = FindSessionInstance( InstanceID );
-	if( ProfilerSession && ProfilerWindow.IsValid())
+	if (ProfilerSession.IsValid())
 	{
-		(*ProfilerSession)->UpdateMetadata( ProfilerClient->GetStatMetaData( InstanceID ) );
+		ProfilerSession->UpdateMetadata( MetaData );
 
-		if( (*ProfilerSession)->GetSessionType() == EProfilerSessionTypes::Live )
+		if (ProfilerSession->GetSessionType() == EProfilerSessionTypes::Live)
 		{
 			RequestFilterAndPresetsUpdateEvent.Broadcast();
 		}
 	}
 }
 
-void FProfilerManager::ProfilerClient_OnLoadedMetaData( const FGuid& InstanceID )
-{
-	//TrackDefaultStats();
-}
+
 
 void FProfilerManager::ProfilerClient_OnLoadStarted( const FGuid& InstanceID )
 {
-	const FProfilerSessionRef* ProfilerSession = FindSessionInstance( InstanceID );
-	if( ProfilerSession && GetProfilerWindow().IsValid())
+	if (ProfilerSession.IsValid() && GetProfilerWindow().IsValid())
 	{
-		const FString Description = (*ProfilerSession)->GetName();
+		const FString Description = ProfilerSession->GetName();
+		UE_LOG( LogStats, Warning, TEXT( "OnLoadStarted: %s" ), *Description );
+		LoadStartTime = FPlatformTime::Seconds();
 
 		// Display the notification that a file is being loaded.
-		GetProfilerWindow()->ManageLoadingProgressNotificationState( (*ProfilerSession)->GetName(), EProfilerNotificationTypes::LoadingOfflineCapture, ELoadingProgressStates::Started, 0.0f );
+		GetProfilerWindow()->ManageLoadingProgressNotificationState( Description, EProfilerNotificationTypes::LoadingOfflineCapture, ELoadingProgressStates::Started, 0.0f );
 	}
 }
 
 void FProfilerManager::ProfilerClient_OnLoadCompleted( const FGuid& InstanceID )
 {
 	// Inform that the file has been loaded and we can hide the notification.
-	const FProfilerSessionRef* ProfilerSession = FindSessionInstance( InstanceID );
-	if( ProfilerSession && GetProfilerWindow().IsValid())
+	if (ProfilerSession.IsValid())
 	{
-		(*ProfilerSession)->LoadComplete();
+		ProfilerSession->LoadComplete();
 
-		// Update the notification that a file has been loaded, to be precise it should be loaded on the next tick...
-		GetProfilerWindow()->ManageLoadingProgressNotificationState( (*ProfilerSession)->GetName(), EProfilerNotificationTypes::LoadingOfflineCapture, ELoadingProgressStates::Loaded, 1.0f );
+		const FString Description = ProfilerSession->GetName();
+		UE_LOG( LogStats, Warning, TEXT( "OnLoadCompleted: %s" ), *Description );
 	}
 }
 
@@ -392,29 +425,26 @@ void FProfilerManager::ProfilerClient_OnClientDisconnected( const FGuid& Session
 	SessionManager
 -----------------------------------------------------------------------------*/
 
-void FProfilerManager::SessionManager_OnCanSelectSession( const ISessionInfoPtr& SelectedSession, bool& CanSelect )
-{
-}
-
 void FProfilerManager::SessionManager_OnSelectedSessionChanged( const ISessionInfoPtr& InActiveSession )
 {
-	SessionManager_OnInstanceSelectionChanged(nullptr, false);
+	//SessionManager_OnInstanceSelectionChanged(nullptr, false);
 }
 
-void FProfilerManager::SessionManager_OnInstanceSelectionChanged(const TSharedPtr<ISessionInstanceInfo>& /*Instance*/, bool /*Selected*/)
+void FProfilerManager::SessionManager_OnInstanceSelectionChanged(const TSharedPtr<ISessionInstanceInfo>& InInstance, bool Selected)
 {
 	const ISessionInfoPtr& SelectedSession = SessionManager->GetSelectedSession();
 	const bool SessionIsValid = SelectedSession.IsValid()
 		&& (SelectedSession->GetSessionOwner() == FPlatformProcess::UserName(false))
 		&& (SessionManager->GetSelectedInstances().Num() > 0);
 
-	if( ActiveSession != SelectedSession || FProfilerManager::GetSettings().bSingleInstanceMode )
+	if (InInstance->GetInstanceId() != ActiveInstanceID && SessionIsValid && InInstance.IsValid() && Selected)
 	{
 		ClearStatsAndInstances();
 
 		if (SessionIsValid)
 		{
 			ActiveSession = SelectedSession;
+			ActiveInstanceID = InInstance->GetInstanceId();
 			ProfilerClient->Subscribe( ActiveSession->GetSessionId() );
 			ProfilerType = EProfilerSessionTypes::Live;
 			SetViewMode( EProfilerViewMode::LineIndexBased );
@@ -422,36 +452,20 @@ void FProfilerManager::SessionManager_OnInstanceSelectionChanged(const TSharedPt
 		else
 		{
 			ActiveSession = nullptr;
+			ProfilerClient->Unsubscribe();
 			ProfilerType = EProfilerSessionTypes::InvalidOrMax;
 		}
-	}
 
-	if( ActiveSession.IsValid() )
-	{
-		// Track all selected session instances.
-		SelectedSessionInstances = SessionManager->GetSelectedInstances();
-		const int32 NumSelectedInstances = SelectedSessionInstances.Num();
-		const int32 NumInstances = FMath::Min( NumSelectedInstances, FProfilerManager::GetSettings().bSingleInstanceMode ? 1 : NumSelectedInstances );
-
-		for( int32 Index = 0; Index < NumInstances; ++Index )
+		if (ActiveSession.IsValid())
 		{
-			const ISessionInstanceInfoPtr SessionInstanceInfo = SelectedSessionInstances[Index];
-			const FGuid ProfilerInstanceID = SessionInstanceInfo->GetInstanceId();
-			const bool bAlreadyAdded = ProfilerSessionInstances.Contains( ProfilerInstanceID );
-
-			if( !bAlreadyAdded )
+			ProfilerSession = MakeShareable( new FProfilerSession( InInstance ) );
+			ProfilerSession->SetOnAddThreadTime( FProfilerSession::FAddThreadTimeDelegate::CreateSP( this, &FProfilerManager::ProfilerSession_OnAddThreadTime ) );
+			ProfilerClient->Track( ActiveInstanceID );
+			TSharedPtr<SProfilerWindow> ProfilerWindowPtr = GetProfilerWindow();
+			if (ProfilerWindowPtr.IsValid())
 			{
-				FProfilerSessionRef ProfilerSession = MakeShareable( new FProfilerSession( SessionInstanceInfo ) );
-				ProfilerSession->SetOnAddThreadTime( FProfilerSession::FAddThreadTimeDelegate::CreateSP( this, &FProfilerManager::ProfilerSession_OnAddThreadTime ) );
-
-				ProfilerSessionInstances.Add( ProfilerSession->GetInstanceID(), ProfilerSession );
-				ProfilerClient->Track( ProfilerInstanceID );
-				TSharedPtr<SProfilerWindow> ProfilerWindowPtr = GetProfilerWindow();
-				if (ProfilerWindowPtr.IsValid())
-				{
-					ProfilerWindowPtr->ManageEventGraphTab(ProfilerInstanceID, true, ProfilerSession->GetName());
-				}
-			}
+				ProfilerWindowPtr->ManageEventGraphTab( ActiveInstanceID, true, ProfilerSession->GetName() );
+			}			
 		}
 	}
 
@@ -477,14 +491,7 @@ bool FProfilerManager::TrackStat( const uint32 StatID )
 	bool bAdded = false;
 
 	// Check if all profiler instances have this stat ready.
-	int32 NumReadyStats = 0;
-	for( auto It = GetProfilerInstancesIterator(); It; ++It )
-	{
-		const FProfilerSessionRef ProfilerSession = It.Value();
-		NumReadyStats += ProfilerSession->GetAggregatedStat(StatID) != nullptr ? 1 : 0;
-	}
-	const bool bStatIsReady = NumReadyStats == GetProfilerInstancesNum();
-
+	const bool bStatIsReady = ProfilerSession->GetAggregatedStat( StatID ) != nullptr;
 	if( StatID != 0 && bStatIsReady )
 	{
 		FTrackedStat* TrackedStat = TrackedStats.Find( StatID );
@@ -517,16 +524,13 @@ bool FProfilerManager::TrackStat( const uint32 StatID )
 			uint32 NumAddedInstances = 0;
 			bool bMetadataInitialized = false;
 
-			for( auto It = GetProfilerInstancesIterator(); It; ++It )
+			if (ProfilerSession.IsValid())
 			{
-				const FGuid& SessionInstanceID = It.Key();// ProfilerSessionInstanceID, ProfilerInstanceID, InstanceID
-				const FProfilerSessionRef ProfilerSession = It.Value();
-				const bool bInstanceAdded = TrackStatForSessionInstance( StatID, SessionInstanceID );
+
+				const bool bInstanceAdded = TrackStatForSessionInstance( StatID, ActiveInstanceID );
 				NumAddedInstances += bInstanceAdded ? 1 : 0;
 
 				// Initialize metadata for combine graph data source.
-				// TODO: This should be checked against the remaining elements to detect inconsistent data.
-				// The first instance should be the main.
 				if( !bMetadataInitialized )
 				{
 					const bool bIsStatReady = ProfilerSession->GetMetaData()->IsStatInitialized( StatID );
@@ -579,12 +583,9 @@ void FProfilerManager::ClearStatsAndInstances()
 		const FTrackedStat& TrackedStat = It.Value();
 		UntrackStat( TrackedStat.StatID );
 	}
-	for( auto It = GetProfilerInstancesIterator(); It; ++It )
-	{
-		const FGuid ProfilerInstanceID = It.Key();
-		ProfilerClient->Untrack( ProfilerInstanceID );
-	}
-	ProfilerSessionInstances.Empty();
+
+	ProfilerClient->Untrack( ActiveInstanceID );
+	ActiveInstanceID.Invalidate();
 }
 
 const bool FProfilerManager::IsStatTracked( const uint32 StatID ) const
@@ -596,18 +597,16 @@ bool FProfilerManager::TrackStatForSessionInstance( const uint32 StatID, const F
 {
 	bool bAdded = false;
 
-	const FProfilerSessionRef* ProfilerSessionPtr = FindSessionInstance( SessionInstanceID );
 	const FTrackedStat* TrackedStat = TrackedStats.Find( StatID );
-
-	if( ProfilerSessionPtr && TrackedStat )
+	if( ProfilerSession.IsValid() && TrackedStat )
 	{
-		const bool bIsStatReady = (*ProfilerSessionPtr)->GetMetaData()->IsStatInitialized( StatID );
+		const bool bIsStatReady = ProfilerSession->GetMetaData()->IsStatInitialized( StatID );
 		const bool bIsExist = TrackedStat->CombinedGraphDataSource->IsProfilerSessionRegistered( SessionInstanceID );
 
 		if( bIsStatReady && !bIsExist )
 		{
 			// Create graph data provider for tracked stat.
-			FGraphDataSourceRefConst	GraphDataSource = (*ProfilerSessionPtr)->CreateGraphDataSource( StatID );
+			FGraphDataSourceRefConst	GraphDataSource = ProfilerSession->CreateGraphDataSource( StatID );
 			TrackedStat->CombinedGraphDataSource->RegisterWithProfilerSession( SessionInstanceID, GraphDataSource );
 			bAdded = true;
 		}
@@ -620,7 +619,6 @@ bool FProfilerManager::UntrackStatForSessionInstance( const uint32 StatID, const
 {
 	bool bRemoved = false;
 
-	//FProfilerSessionRef* ProfilerSessionPtr = FindSessionInstance( SessionInstanceID );
 	const FTrackedStat* TrackedStat = TrackedStats.Find( StatID );
 
 	if( TrackedStat )
@@ -648,46 +646,6 @@ const bool FProfilerManager::IsStatTrackedForSessionInstance( const uint32 StatI
 	return bResult;
 }
 
-bool FProfilerManager::TrackSessionInstance( const FGuid& SessionInstanceID )
-{
-	for( auto It = TrackedStats.CreateConstIterator(); It; ++It )
-	{
-		const FTrackedStat& TrackedStat = It.Value();
-		TrackStatForSessionInstance( TrackedStat.StatID, SessionInstanceID );
-	}
-	return true;
-}
-
-bool FProfilerManager::UntrackSessionInstance( const FGuid& SessionInstanceID )
-{
-	for( auto It = TrackedStats.CreateConstIterator(); It; ++It )
-	{
-		const FTrackedStat& TrackedStat = It.Value();
-		UntrackStatForSessionInstance( TrackedStat.StatID, SessionInstanceID );
-	}
-	return true;
-}
-
-const bool FProfilerManager::IsSessionInstanceTracked( const FGuid& SessionInstanceID ) const
-{
-	const FProfilerSessionRef* ProfilerSessionPtr = FindSessionInstance( SessionInstanceID );
-	if( ProfilerSessionPtr )
-	{
-		for( auto It = TrackedStats.CreateConstIterator(); It; ++It )
-		{
-			const FTrackedStat& TrackedStat = It.Value();
-			const bool bIsSessionInstanceTracked = TrackedStat.CombinedGraphDataSource->IsProfilerSessionRegistered( SessionInstanceID );
-			if( bIsSessionInstanceTracked )
-			{
-				return true;
-			}
-		}
-	}
-
-	return false;
-}
-
-//-----------------------------------------------------------------------------
 
 // @TODO: Move to profiler settings
 const FLinearColor& FProfilerManager::GetColorForStatID( const uint32 StatID ) const
@@ -707,14 +665,12 @@ const FLinearColor& FProfilerManager::GetColorForStatID( const uint32 StatID ) c
 void FProfilerManager::TrackDefaultStats()
 {
 	// Find StatId for the game thread.
-	for( auto It = GetProfilerInstancesIterator(); It; ++It )
+	if (ProfilerSession.IsValid())
 	{
-		FProfilerSessionRef ProfilerSession = It.Value();
 		if( ProfilerSession->GetMetaData()->IsReady() )
 		{
 			TrackStat( ProfilerSession->GetMetaData()->GetGameThreadStatID() );
 		}
-		break;
 	}
 }
 
@@ -728,9 +684,8 @@ void FProfilerManager::CloseAllEventGraphTabs()
 	if( ProfilerWindowPtr.IsValid() )
 	{
 		// Iterate through all profiler sessions.
-		for( auto It = GetProfilerInstancesIterator(); It; ++It )
+		if (ProfilerSession.IsValid())
 		{
-			FProfilerSessionRef ProfilerSession = It.Value();
 			ProfilerWindowPtr->ManageEventGraphTab( ProfilerSession->GetInstanceID(), false, TEXT("") );
 		}
 
@@ -742,13 +697,10 @@ void FProfilerManager::DataGraph_OnSelectionChangedForIndex( uint32 FrameStartIn
 {
 	PROFILER_SCOPE_LOG_TIME( TEXT( "FProfilerManager::DataGraph_OnSelectionChangedForIndex" ), nullptr );
 
-	for( auto It = GetProfilerInstancesIterator(); It; ++It )
+	if (ProfilerSession.IsValid())
 	{
-		FProfilerSessionRef ProfilerSession = It.Value();
-	
 		FEventGraphDataRef EventGraphDataAverage = ProfilerSession->CreateEventGraphData( FrameStartIndex, FrameEndIndex, EEventGraphTypes::Average );
 		FEventGraphDataRef EventGraphDataMaximum = ProfilerSession->CreateEventGraphData( FrameStartIndex, FrameEndIndex, EEventGraphTypes::Maximum );
-
 		GetProfilerWindow()->UpdateEventGraph( ProfilerSession->GetInstanceID(), EventGraphDataAverage, EventGraphDataMaximum, false );
 	}
 }
@@ -759,6 +711,10 @@ void FProfilerManager::ProfilerSession_OnAddThreadTime( int32 FrameIndex, const
 	if( ProfilerWindowPtr.IsValid() )
 	{
 		ProfilerWindowPtr->ProfilerMiniView->AddThreadTime( FrameIndex, ThreadMS, StatMetaData );
+		
+		// Update the notification that a file is being loaded.
+		const float DataLoadingProgress = ProfilerSession->GetProgress();
+		ProfilerWindowPtr->ManageLoadingProgressNotificationState( ProfilerSession->GetName(), EProfilerNotificationTypes::LoadingOfflineCapture, ELoadingProgressStates::InProgress, DataLoadingProgress );		
 	}
 }
 
diff --git a/Engine/Source/Developer/Profiler/Private/ProfilerManager.h b/Engine/Source/Developer/Profiler/Private/ProfilerManager.h
--- a/Engine/Source/Developer/Profiler/Private/ProfilerManager.h
+++ b/Engine/Source/Developer/Profiler/Private/ProfilerManager.h
@@ -11,8 +11,7 @@ class FProfilerSettings
 {
 public:
 	FProfilerSettings( bool bInIsDefault = false )
-		: bSingleInstanceMode( true )
-		, bShowCoalescedViewModesInEventGraph( true )
+		: bShowCoalescedViewModesInEventGraph( true )
 		, bIsEditing( false )
 		, bIsDefault( bInIsDefault )
 	{
@@ -34,13 +33,11 @@ public:
 	{
 		FConfigCacheIni::LoadGlobalIniFile(ProfilerSettingsIni, TEXT("ProfilerSettings"));
 
-		GConfig->GetBool( TEXT("Profiler.ProfilerOptions"), TEXT("bSingleInstanceMode"), bSingleInstanceMode, ProfilerSettingsIni );
 		GConfig->GetBool( TEXT("Profiler.ProfilerOptions"), TEXT("bShowCoalescedViewModesInEventGraph"), bShowCoalescedViewModesInEventGraph, ProfilerSettingsIni );
 	}
 
 	void SaveToConfig()
 	{
-		GConfig->SetBool( TEXT("Profiler.ProfilerOptions"), TEXT("bSingleInstanceMode"), bSingleInstanceMode, ProfilerSettingsIni );
 		GConfig->SetBool( TEXT("Profiler.ProfilerOptions"), TEXT("bShowCoalescedViewModesInEventGraph"), bShowCoalescedViewModesInEventGraph, ProfilerSettingsIni );
 		GConfig->Flush( false, ProfilerSettingsIni );
 	}
@@ -72,9 +69,6 @@ public:
 	/** Profiler setting filename ini. */
 	FString ProfilerSettingsIni;
 
-	/** If True, the profiler will work in the single instance mode, all functionality related to multi instances will be disabled or removed from the UI. */
-	bool bSingleInstanceMode;
-
 	/** If True, coalesced view modes related functionality will be added to the event graph. */
 	bool bShowCoalescedViewModesInEventGraph;
 
@@ -258,6 +252,14 @@ public:
 		return FProfilerCommands::Get();
 	}
 
+	/**
+	 * @return an instance of the profiler session.
+	 */
+	FProfilerSessionPtr GetProfilerSession()
+	{
+		return ProfilerSession;
+	}
+
 	/**
 	 * Creates a combined graph data source which will provide data for graph drawing.
 	 *
@@ -282,31 +284,6 @@ public:
 	bool TrackStatForSessionInstance( const uint32 StatID, const FGuid& SessionInstanceID );
 	bool UntrackStatForSessionInstance( const uint32 StatID, const FGuid& SessionInstanceID );
 	const bool IsStatTrackedForSessionInstance( const uint32 StatID, const FGuid& SessionInstanceID ) const;
-	bool TrackSessionInstance( const FGuid& SessionInstanceID );
-	bool UntrackSessionInstance( const FGuid& SessionInstanceID );
-
-	/**
-	 * @return true, if the specified session instance ID belongs to a valid profiler session.
-	 */
-	bool IsSessionInstanceValid( const FGuid& SessionInstanceID ) const
-	{
-		const bool bIsValid = SessionInstanceID.IsValid() && ProfilerSessionInstances.Contains( SessionInstanceID );
-		return bIsValid;
-	}
-
-	const FProfilerSessionRef* FindSessionInstance( const FGuid& SessionInstanceID ) const
-	{
-		return ProfilerSessionInstances.Find( SessionInstanceID );
-	}
-
-	// TODO: At this moment SButtonRowBlock::OnIsChecked supports only Checked and Unchecked
-	// ECheckBoxState.Undetermined is not supported
-	/**
-	 * @return true, if the specified session instance tracks any stats.
-	 */
-	const bool IsSessionInstanceTracked( const FGuid& SessionInstanceID ) const;
-
-	//-----------------------------------------------------------------------------
 
 	/**
 	 * @return True, if the profiler has at least one fully processed capture file
@@ -321,7 +298,13 @@ public:
 	 */
 	const bool IsConnected() const
 	{
-		return ActiveSession.IsValid();
+		const bool bIsValid = ActiveSession.IsValid() && ActiveInstanceID.IsValid();
+		return bIsValid;
+	}
+
+	const bool HasValidSession() const
+	{
+		return ProfilerSession.IsValid();
 	}
 
 	/**
@@ -334,24 +317,7 @@ public:
 
 public:
 	/** @return true, if all session instances are previewing data */
-	const bool IsDataPreviewing()
-	{
-		return GetProfilerInstancesNum() > 0 && GetNumDataPreviewingInstances() == GetProfilerInstancesNum();
-	}
-
-	/** @return the number of session instances with data previewing enabled */
-	const int32 GetNumDataPreviewingInstances()
-	{
-		int32 NumDataPreviewingInstances = 0;
-
-		for( auto It = GetProfilerInstancesIterator(); It; ++It )
-		{
-			FProfilerSessionRef ProfilerSession = It.Value();
-			NumDataPreviewingInstances += ProfilerSession->bDataPreviewing ? 1 : 0;
-		}
-
-		return NumDataPreviewingInstances;
-	}
+	const bool IsDataPreviewing() const;
 
 	/**
 	 * Sets the data preview state for all session instances and sends message for remote profiler services.
@@ -359,35 +325,10 @@ public:
 	 * @param bRequestedDataPreviewState - data preview state that should be set
 	 *
 	 */
-	void SetDataPreview( const bool bRequestedDataPreviewState )
-	{
-		ProfilerClient->SetPreviewState( bRequestedDataPreviewState );
-		for( auto It = GetProfilerInstancesIterator(); It; ++It )
-		{
-			FProfilerSessionRef ProfilerSession = It.Value();
-			ProfilerSession->bDataPreviewing = bRequestedDataPreviewState;	
-		}
-	}
+	void SetDataPreview( const bool bRequestedDataPreviewState );
 
 	/** @return true, if all sessions instances are capturing data to a file, only valid if profiler is connected to network based session */
-	const bool IsDataCapturing()
-	{
-		return GetProfilerInstancesNum() > 0 && GetNumDataCapturingInstances() == GetProfilerInstancesNum();
-	}
-
-	/** @return the number of session instances with data capturing enabled. */
-	const int32 GetNumDataCapturingInstances()
-	{
-		int32 NumDataCapturingInstances = 0;
-
-		for( auto It = GetProfilerInstancesIterator(); It; ++It )
-		{
-			FProfilerSessionRef ProfilerSession = It.Value();
-			NumDataCapturingInstances += ProfilerSession->bDataCapturing ? 1 : 0;
-		}
-
-		return NumDataCapturingInstances;
-	}
+	const bool IsDataCapturing() const;
 
 	/**
 	 * Sets the data capture state for all session instances and sends message for remote profiler services.
@@ -510,40 +451,25 @@ public:
 	 * Creates a new profiler session instance and loads a saved profiler capture from the specified location.
 	 *
 	 * @param ProfilerCaptureFilepath	- The path to the file containing a captured session instance
-	 * @param bAdd						- if true, it will load a captured session instance and add to the existing ones
-	 *
 	 */
-	void LoadProfilerCapture( const FString& ProfilerCaptureFilepath, const bool bAdd = false );
+	void LoadProfilerCapture( const FString& ProfilerCaptureFilepath );
 
 	/** Creates a new profiler session instance and load a raw stats file from the specified location. */
 	void LoadRawStatsFile( const FString& RawStatsFileFileath );
 
 protected:
-	void ProfilerClient_OnProfilerData( const FGuid& InstanceID, const FProfilerDataFrame& Content, const float DataLoadingProgress );
+	void ProfilerClient_OnProfilerData( const FGuid& InstanceID, const FProfilerDataFrame& Content );
 	void ProfilerClient_OnClientConnected( const FGuid& SessioID, const FGuid& InstanceID );
 	void ProfilerClient_OnClientDisconnected( const FGuid& SessionID, const FGuid& InstanceID );
-	void ProfilerClient_OnMetaDataUpdated( const FGuid& InstanceID );
-	void ProfilerClient_OnLoadedMetaData( const FGuid& InstanceID );
+	void ProfilerClient_OnMetaDataUpdated( const FGuid& InstanceID, const FStatMetaData& MetaData );
 	void ProfilerClient_OnLoadCompleted( const FGuid& InstanceID );
 	void ProfilerClient_OnLoadStarted( const FGuid& InstanceID );
 
 	void ProfilerClient_OnProfilerFileTransfer( const FString& Filename, int64 FileProgress, int64 FileSize );
 
-	void SessionManager_OnCanSelectSession( const ISessionInfoPtr& Session, bool& CanSelect );
 	void SessionManager_OnInstanceSelectionChanged( const TSharedPtr<ISessionInstanceInfo>& Instance, bool Selected );
 	void SessionManager_OnSelectedSessionChanged( const ISessionInfoPtr& Session );
 
-public:
-	TMap<FGuid,FProfilerSessionRef>::TIterator GetProfilerInstancesIterator()
-	{
-		return ProfilerSessionInstances.CreateIterator();
-	}
-
-	const int32 GetProfilerInstancesNum() const
-	{
-		return ProfilerSessionInstances.Num();
-	}
-
 public:
 	const FLinearColor& GetColorForStatID( const uint32 StatID ) const;
 
@@ -581,8 +507,11 @@ protected:
 	/** A shared pointer to the currently selected session in the session browser. */
 	ISessionInfoPtr ActiveSession;
 
-	/** Session instances currently selected in the session browser. */
-	TArray<ISessionInstanceInfoPtr> SelectedSessionInstances;
+	/** A shared pointer to the currently selected instance in the session browser. */
+	FGuid/*ISessionInstanceInfoPtr*/ ActiveInstanceID;
+
+	/** Profiler session, to be removed from here. */
+	FProfilerSessionPtr ProfilerSession;
 
 	/** A shared pointer to the profiler client, which is used to deliver all profiler data from the active session. */
 	IProfilerClientPtr ProfilerClient;
@@ -609,9 +538,7 @@ protected:
 	/** Contains all currently tracked stats, stored as StatID -> FTrackedStat. */
 	TMap<uint32, FTrackedStat> TrackedStats;
 
-	/** Holds all profiler session instances, stored as FGuid -> FProfilerSessionRef. */
-	TMap<FGuid,FProfilerSessionRef> ProfilerSessionInstances;
-
+	
 	/*-----------------------------------------------------------------------------
 		Profiler manager states
 	-----------------------------------------------------------------------------*/
diff --git a/Engine/Source/Developer/Profiler/Private/ProfilerModule.cpp b/Engine/Source/Developer/Profiler/Private/ProfilerModule.cpp
--- a/Engine/Source/Developer/Profiler/Private/ProfilerModule.cpp
+++ b/Engine/Source/Developer/Profiler/Private/ProfilerModule.cpp
@@ -69,7 +69,7 @@ void FProfilerModule::Shutdown( TSharedRef<SDockTab> TabBeingClosed )
 
 void FProfilerModule::StatsMemoryDumpCommand( const TCHAR* Filename )
 {
-	TUniquePtr<FRawStatsMemoryProfiler> Instance( FCreateStatsReader<FRawStatsMemoryProfiler>::ForRawStats( Filename ) );
+	TUniquePtr<FRawStatsMemoryProfiler> Instance( FStatsReader<FRawStatsMemoryProfiler>::Create( Filename ) );
 	if (Instance)
 	{
 		Instance->ReadAndProcessSynchronously();
@@ -140,7 +140,7 @@ void FProfilerModule::StatsMemoryDumpCommand( const TCHAR* Filename )
 
 FRawStatsMemoryProfiler* FProfilerModule::OpenRawStatsForMemoryProfiling( const TCHAR* Filename )
 {
-	FRawStatsMemoryProfiler* Instance = FCreateStatsReader<FRawStatsMemoryProfiler>::ForRawStats( Filename );
+	FRawStatsMemoryProfiler* Instance = FStatsReader<FRawStatsMemoryProfiler>::Create( Filename );
 	if (Instance)
 	{
 		Instance->ReadAndProcessAsynchronously();
diff --git a/Engine/Source/Developer/Profiler/Private/ProfilerRawStats.h b/Engine/Source/Developer/Profiler/Private/ProfilerRawStats.h
--- a/Engine/Source/Developer/Profiler/Private/ProfilerRawStats.h
+++ b/Engine/Source/Developer/Profiler/Private/ProfilerRawStats.h
@@ -213,7 +213,7 @@ protected:
 /** Class for managing the raw stats in the context of memory profiling. */
 class FRawStatsMemoryProfiler : public FStatsReadFile
 {
-	friend struct FCreateStatsReader<FRawStatsMemoryProfiler>;
+	friend struct FStatsReader<FRawStatsMemoryProfiler>;
 	typedef FStatsReadFile Super;
 
 protected:
diff --git a/Engine/Source/Developer/Profiler/Private/ProfilerSession.cpp b/Engine/Source/Developer/Profiler/Private/ProfilerSession.cpp
--- a/Engine/Source/Developer/Profiler/Private/ProfilerSession.cpp
+++ b/Engine/Source/Developer/Profiler/Private/ProfilerSession.cpp
@@ -21,8 +21,7 @@ FProfilerStat::FProfilerStat( const uint32 InStatID /*= 0 */ )
 -----------------------------------------------------------------------------*/
 
 FProfilerSession::FProfilerSession( EProfilerSessionTypes::Type InSessionType, const ISessionInstanceInfoPtr InSessionInstanceInfo, FGuid InSessionInstanceID, FString InDataFilepath )
-: ClientStatMetadata( nullptr )
-, bRequestStatMetadataUpdate( false )
+: bRequestStatMetadataUpdate( false )
 , bLastPacket( false )
 , StatMetaDataSize( 0 )
 , OnTick( FTickerDelegate::CreateRaw( this, &FProfilerSession::HandleTicker ) )
@@ -30,12 +29,14 @@ FProfilerSession::FProfilerSession( EProfilerSessionTypes::Type InSessionType, c
 , StatMetaData( MakeShareable( new FProfilerStatMetaData() ) )
 , EventGraphDataTotal( MakeShareable( new FEventGraphData() ) )
 , EventGraphDataMaximum( MakeShareable( new FEventGraphData() ) )
-, EventGraphDataCurrent( MakeShareable( new FEventGraphData() ) )
+, EventGraphDataCurrent( nullptr )
 , CreationTime( FDateTime::Now() )
 , SessionType( InSessionType )
 , SessionInstanceInfo( InSessionInstanceInfo )
 , SessionInstanceID( InSessionInstanceID )
 , DataFilepath( InDataFilepath )
+, NumFrames( 0 )
+, NumFramesProcessed( 0 )
 , bDataPreviewing( false )
 , bDataCapturing( false )
 , bHasAllProfilerData( false )
@@ -44,8 +45,7 @@ FProfilerSession::FProfilerSession( EProfilerSessionTypes::Type InSessionType, c
 {}
 
 FProfilerSession::FProfilerSession( const ISessionInstanceInfoPtr InSessionInstanceInfo )
-: ClientStatMetadata( nullptr )
-, bRequestStatMetadataUpdate( false )
+: bRequestStatMetadataUpdate( false )
 , bLastPacket( false )
 , StatMetaDataSize( 0 )
 , OnTick( FTickerDelegate::CreateRaw( this, &FProfilerSession::HandleTicker ) )
@@ -53,12 +53,14 @@ FProfilerSession::FProfilerSession( const ISessionInstanceInfoPtr InSessionInsta
 , StatMetaData( MakeShareable( new FProfilerStatMetaData() ) )
 , EventGraphDataTotal( MakeShareable( new FEventGraphData() ) )
 , EventGraphDataMaximum( MakeShareable( new FEventGraphData() ) )
-, EventGraphDataCurrent( MakeShareable( new FEventGraphData() ) )
+, EventGraphDataCurrent( nullptr )
 , CreationTime( FDateTime::Now() )
 , SessionType( EProfilerSessionTypes::Live )
 , SessionInstanceInfo( InSessionInstanceInfo )
 , SessionInstanceID( InSessionInstanceInfo->GetInstanceId() )
 , DataFilepath( TEXT("") )
+, NumFrames( 0 )
+, NumFramesProcessed( 0 )
 , bDataPreviewing( false )
 , bDataCapturing( false )
 , bHasAllProfilerData( false )
@@ -71,8 +73,7 @@ FProfilerSession::FProfilerSession( const ISessionInstanceInfoPtr InSessionInsta
 }
 
 FProfilerSession::FProfilerSession( const FString& InDataFilepath )
-: ClientStatMetadata( nullptr )
-, bRequestStatMetadataUpdate( false )
+: bRequestStatMetadataUpdate( false )
 , bLastPacket( false )
 , StatMetaDataSize( 0 )
 , OnTick( FTickerDelegate::CreateRaw( this, &FProfilerSession::HandleTicker ) )
@@ -80,12 +81,14 @@ FProfilerSession::FProfilerSession( const FString& InDataFilepath )
 , StatMetaData( MakeShareable( new FProfilerStatMetaData() ) )
 , EventGraphDataTotal( MakeShareable( new FEventGraphData() ) )
 , EventGraphDataMaximum( MakeShareable( new FEventGraphData() ) )
-, EventGraphDataCurrent( MakeShareable( new FEventGraphData() ) )
+, EventGraphDataCurrent( nullptr )
 , CreationTime( FDateTime::Now() )
 , SessionType( EProfilerSessionTypes::StatsFile )
 , SessionInstanceInfo( nullptr )
 , SessionInstanceID( FGuid::NewGuid() )
 , DataFilepath( InDataFilepath.Replace( *FStatConstants::StatsFileExtension, TEXT( "" ) ) )
+, NumFrames( 0 )
+, NumFramesProcessed( 0 )
 , bDataPreviewing( false )
 , bDataCapturing( false )
 , bHasAllProfilerData( false )
@@ -143,6 +146,9 @@ FEventGraphDataRef FProfilerSession::CreateEventGraphData( const uint32 FrameSta
 
 void FProfilerSession::UpdateAggregatedStats( const uint32 FrameIndex )
 {
+	static FTotalTimeAndCount TimeAndCount( 0.0f, 0 );
+	PROFILER_SCOPE_LOG_TIME( TEXT( "2 FProfilerSession::UpdateAggregatedStats" ), &TimeAndCount );
+
 	const FIntPoint& IndicesForFrame = DataProvider->GetSamplesIndicesForFrame( FrameIndex );
 	const uint32 SampleStartIndex = IndicesForFrame.X;
 	const uint32 SampleEndIndex = IndicesForFrame.Y;
@@ -172,34 +178,37 @@ void FProfilerSession::UpdateAggregatedStats( const uint32 FrameIndex )
 	// @TODO: Create map for stats TMap<uint32, TArray<uint32> >; StatID -> Sample indices for faster lookup in data providers
 }
 
-void FProfilerSession::EventGraphCombineAndMax( const FEventGraphDataRef Current, const uint32 NumFrames )
+void FProfilerSession::EventGraphCombineAndMax( const FEventGraphData* Current, const uint32 InNumFrames )
 {
-	EventGraphDataMaximum->CombineAndFindMax( Current.Get() );
-	EventGraphDataMaximum->Advance( 0, NumFrames );
+	EventGraphDataMaximum->CombineAndFindMax( *Current );
+	
+	if (InNumFrames > 0)
+	{
+		EventGraphDataMaximum->Advance( 0, InNumFrames );
+	}
 }
 
-void FProfilerSession::EventGraphCombineAndAdd( const FEventGraphDataRef Current, const uint32 NumFrames )
+void FProfilerSession::EventGraphCombineAndAdd( const FEventGraphData* Current, const uint32 InNumFrames )
 {
-	EventGraphDataTotal->CombineAndAdd( Current.Get() );
-	EventGraphDataTotal->Advance( 0, NumFrames );
+	EventGraphDataTotal->CombineAndAdd( *Current );
+
+	if (InNumFrames > 0)
+	{
+		EventGraphDataTotal->Advance( 0, InNumFrames );
+	}
 }
 
 void FProfilerSession::UpdateAggregatedEventGraphData( const uint32 FrameIndex )
 {
-	if( CompletionSync.GetReference() && !CompletionSync->IsComplete() )
-	{
-		static FTotalTimeAndCount JoinTasksTimeAndCont(0.0f, 0);
-		PROFILER_SCOPE_LOG_TIME( TEXT( "FProfilerSession::CombineJoinAndContinue" ), &JoinTasksTimeAndCont );
-
-		FTaskGraphInterface::Get().WaitUntilTaskCompletes(CompletionSync, ENamedThreads::GameThread);
-	}
+	static FTotalTimeAndCount TimeAndCount( 0.0f, 0 );
+	PROFILER_SCOPE_LOG_TIME( TEXT( "3  FProfilerSession::UpdateAggregatedEventGraphData" ), &TimeAndCount );
 
-	static FTotalTimeAndCount TimeAndCount(0.0f, 0);
-	PROFILER_SCOPE_LOG_TIME( TEXT( "FProfilerSession::UpdateAggregatedEventGraphData" ), &TimeAndCount );
+	CompletionSyncAggregatedEventGraphData();
 
 	// Create a temporary event graph data for the specified frame.
-	EventGraphDataCurrent = MakeShareable( new FEventGraphData( AsShared(), FrameIndex ) );
-	const uint32 NumFrames = DataProvider->GetNumFrames();
+	delete EventGraphDataCurrent;
+	EventGraphDataCurrent = new FEventGraphData( AsShared(), FrameIndex );
+	const uint32 NumFramesLocal = SessionType == EProfilerSessionTypes::Live ? DataProvider->GetNumFrames() : 0;
 
 	static const bool bUseTaskGraph = true;
 
@@ -210,19 +219,19 @@ void FProfilerSession::UpdateAggregatedEventGraphData( const uint32 FrameIndex )
 		DECLARE_CYCLE_STAT(TEXT("FSimpleDelegateGraphTask.EventGraphData.CombineAndFindMax"),
 			STAT_FSimpleDelegateGraphTask_EventGraphData_CombineAndFindMax,
 			STATGROUP_TaskGraphTasks);
-		DECLARE_CYCLE_STAT(TEXT("FSimpleDelegateGraphTask.EventGraphData.EventGraphCombineAndAdd"),
+		DECLARE_CYCLE_STAT(TEXT("FSimpleDelegateGraphTask.EventGraphData.CombineAndAdd"),
 			STAT_FSimpleDelegateGraphTask_EventGraphData_EventGraphCombineAndAdd,
 			STATGROUP_TaskGraphTasks);
 
-		new (EventGraphCombineTasks)FGraphEventRef(FSimpleDelegateGraphTask::CreateAndDispatchWhenReady
+		new (EventGraphCombineTasks) FGraphEventRef(FSimpleDelegateGraphTask::CreateAndDispatchWhenReady
 		(
-			FSimpleDelegateGraphTask::FDelegate::CreateRaw( this, &FProfilerSession::EventGraphCombineAndMax, EventGraphDataCurrent, NumFrames ), 
+			FSimpleDelegateGraphTask::FDelegate::CreateRaw( this, &FProfilerSession::EventGraphCombineAndMax, EventGraphDataCurrent, NumFramesLocal ), 
 			GET_STATID(STAT_FSimpleDelegateGraphTask_EventGraphData_CombineAndFindMax), nullptr
 		));
 
 		new (EventGraphCombineTasks) FGraphEventRef(FSimpleDelegateGraphTask::CreateAndDispatchWhenReady
 		(
-			FSimpleDelegateGraphTask::FDelegate::CreateRaw( this, &FProfilerSession::EventGraphCombineAndAdd, EventGraphDataCurrent, NumFrames ), 
+			FSimpleDelegateGraphTask::FDelegate::CreateRaw( this, &FProfilerSession::EventGraphCombineAndAdd, EventGraphDataCurrent, NumFramesLocal ), 
 			GET_STATID(STAT_FSimpleDelegateGraphTask_EventGraphData_EventGraphCombineAndAdd), nullptr
 		));
 
@@ -236,25 +245,45 @@ void FProfilerSession::UpdateAggregatedEventGraphData( const uint32 FrameIndex )
 	}
 	else
 	{
-		EventGraphCombineAndMax( EventGraphDataCurrent, NumFrames );
-		EventGraphCombineAndAdd( EventGraphDataCurrent, NumFrames );
+		EventGraphCombineAndMax( EventGraphDataCurrent, NumFramesLocal );
+		EventGraphCombineAndAdd( EventGraphDataCurrent, NumFramesLocal );
 	}
 }
 
-bool FProfilerSession::HandleTicker( float DeltaTime )
+void FProfilerSession::CompletionSyncAggregatedEventGraphData()
 {
-	static FTotalTimeAndCount Current( 0.0f, 0 );
-	PROFILER_SCOPE_LOG_TIME( TEXT( "FProfilerSession::HandleTicker" ), &Current );
+	if (CompletionSync.GetReference() && !CompletionSync->IsComplete())
+	{
+		static FTotalTimeAndCount JoinTasksTimeAndCont( 0.0f, 0 );
+		PROFILER_SCOPE_LOG_TIME( TEXT( "4   FProfilerSession::CombineJoinAndContinue" ), &JoinTasksTimeAndCont );
+
+		FTaskGraphInterface::Get().WaitUntilTaskCompletes( CompletionSync, ENamedThreads::GameThread );
+	}
+}
 
+bool FProfilerSession::HandleTicker( float DeltaTime )
+{
 	// Update metadata if needed
 	if( bRequestStatMetadataUpdate )
 	{
-		StatMetaData->Update( *ClientStatMetadata );
+		StatMetaData->Update( ClientStatMetadata );
 		bRequestStatMetadataUpdate = false;
 	}
 
-	// Limit processing to 30ms per frame.
-	const double TimeLimit = 30 / 1000.0;
+	static double ProcessingTime = 1.0;
+	ProcessingTime -= DeltaTime;
+
+	static int32 NumFramesProcessedLastTime = 0;
+
+	if (ProcessingTime < 0.0)
+	{
+		UE_LOG( LogStats, Verbose, TEXT( "NumFramesProcessedLastTime: %4i / %4i" ), NumFramesProcessedLastTime, FrameToProcess.Num() );
+		ProcessingTime = 1.0;
+		NumFramesProcessedLastTime = 0;
+	}
+
+	// Limit processing to 50ms per frame.
+	const double TimeLimit = 50 / 1000.0;
 	double Seconds = 0;
 	
 	for( int32 FrameNumber = 0; FrameNumber < FrameToProcess.Num(); FrameNumber++ )
@@ -264,6 +293,12 @@ bool FProfilerSession::HandleTicker( float DeltaTime )
 			break;
 		}
 
+		static FTotalTimeAndCount Current( 0.0f, 0 );
+		PROFILER_SCOPE_LOG_TIME( TEXT( "1 FProfilerSession::HandleTicker" ), &Current );
+
+		NumFramesProcessedLastTime++;
+		NumFramesProcessed++;
+
 		FSimpleScopeSecondsCounter SecondsCounter(Seconds);
 
 		const uint32 FrameIndex = FrameToProcess[0];
@@ -271,7 +306,6 @@ bool FProfilerSession::HandleTicker( float DeltaTime )
 
 		FProfilerDataFrame& CurrentProfilerData = FrameToProfilerDataMapping.FindChecked( FrameIndex );
 
-		const FProfilerStatMetaDataRef MetaData = GetMetaData();
 		TMap<uint32, float> ThreadMS;
 
 		// Preprocess the hierarchical samples for the specified frame.
@@ -279,7 +313,7 @@ bool FProfilerSession::HandleTicker( float DeltaTime )
 
 		// Add a root sample for this frame.
 		// HACK 2013-07-19, 13:44 STAT_Root == ThreadRoot in stats2
-		const uint32 FrameRootSampleIndex = DataProvider->AddHierarchicalSample( 0, MetaData->GetStatByID(1).OwningGroup().ID(), 1, 0.0f, 0.0f, 1 );
+		const uint32 FrameRootSampleIndex = DataProvider->AddHierarchicalSample( 0, StatMetaData->GetStatByID(1).OwningGroup().ID(), 1, 0.0f, 0.0f, 1 );
 		
 		double GameThreadTimeMS = 0.0f;
 		double MaxThreadTimeMS = 0.0f;
@@ -295,25 +329,25 @@ bool FProfilerSession::HandleTicker( float DeltaTime )
 			ThreadStartTimeMS = CurrentProfilerData.FrameStart;
 			for( int32 Index = 0; Index < ThreadGraph.Children.Num(); Index++ )
 			{
-				ThreadDurationTimeMS += MetaData->ConvertCyclesToMS( ThreadGraph.Children[Index].Value );
+				ThreadDurationTimeMS += StatMetaData->ConvertCyclesToMS( ThreadGraph.Children[Index].Value );
 			}
 
 			if (ThreadDurationTimeMS > 0.0)
 			{
 				// Check for game thread.
 				const FString GameThreadName = FName( NAME_GameThread ).GetPlainNameString();
-				const bool bGameThreadFound = MetaData->GetThreadDescriptions().FindChecked( ThreadID ).Contains( GameThreadName );
+				const bool bGameThreadFound = StatMetaData->GetThreadDescriptions().FindChecked( ThreadID ).Contains( GameThreadName );
 				if( bGameThreadFound )
 				{
 					GameThreadTimeMS = ThreadDurationTimeMS;
 				}
 
 				// Add a root sample for each thread.
-				const uint32& NewThreadID = MetaData->ThreadIDtoStatID.FindChecked( ThreadID );
+				const uint32& NewThreadID = StatMetaData->ThreadIDtoStatID.FindChecked( ThreadID );
 				const uint32 ThreadRootSampleIndex = DataProvider->AddHierarchicalSample
 				( 
 					NewThreadID/*ThreadID*/, 
-					MetaData->GetStatByID(NewThreadID).OwningGroup().ID(), 
+					StatMetaData->GetStatByID(NewThreadID).OwningGroup().ID(), 
 					NewThreadID, 
 					ThreadStartTimeMS, 
 					ThreadDurationTimeMS, 
@@ -326,7 +360,7 @@ bool FProfilerSession::HandleTicker( float DeltaTime )
 				for( int32 Index = 0; Index < ThreadGraph.Children.Num(); Index++ )
 				{
 					const FProfilerCycleGraph& CycleGraph = ThreadGraph.Children[Index];
-					const double CycleDurationMS = MetaData->ConvertCyclesToMS( CycleGraph.Value );
+					const double CycleDurationMS = StatMetaData->ConvertCyclesToMS( CycleGraph.Value );
 					const double CycleStartTimeMS = ThreadStartTimeMS;
 
 					if (CycleDurationMS > 0.0)
@@ -352,15 +386,15 @@ bool FProfilerSession::HandleTicker( float DeltaTime )
 			for( int32 Index = 0; Index < CurrentProfilerData.CountAccumulators.Num(); Index++ )
 			{
 				const FProfilerCountAccumulator& IntCounter = CurrentProfilerData.CountAccumulators[Index];
-				const EProfilerSampleTypes::Type ProfilerSampleType = MetaData->GetSampleTypeForStatID( IntCounter.StatId );
-				DataProvider->AddCounterSample( MetaData->GetStatByID(IntCounter.StatId).OwningGroup().ID(), IntCounter.StatId, (double)IntCounter.Value, ProfilerSampleType );
+				const EProfilerSampleTypes::Type ProfilerSampleType = StatMetaData->GetSampleTypeForStatID( IntCounter.StatId );
+				DataProvider->AddCounterSample( StatMetaData->GetStatByID(IntCounter.StatId).OwningGroup().ID(), IntCounter.StatId, (double)IntCounter.Value, ProfilerSampleType );
 			}
 
 			// Process floating point counters.
 			for( int32 Index = 0; Index < CurrentProfilerData.FloatAccumulators.Num(); Index++ )
 			{
 				const FProfilerFloatAccumulator& FloatCounter = CurrentProfilerData.FloatAccumulators[Index];
-				DataProvider->AddCounterSample( MetaData->GetStatByID(FloatCounter.StatId).OwningGroup().ID(), FloatCounter.StatId, (double)FloatCounter.Value, EProfilerSampleTypes::NumberFloat );
+				DataProvider->AddCounterSample( StatMetaData->GetStatByID(FloatCounter.StatId).OwningGroup().ID(), FloatCounter.StatId, (double)FloatCounter.Value, EProfilerSampleTypes::NumberFloat );
 			}
 		}
 
@@ -374,8 +408,13 @@ bool FProfilerSession::HandleTicker( float DeltaTime )
 		// Update aggregated events.
 		UpdateAggregatedEventGraphData( LastFrameIndex );
 
-		// Update mini-view.
-		OnAddThreadTime.ExecuteIfBound( LastFrameIndex, ThreadMS, StatMetaData );
+		{
+			static FTotalTimeAndCount TotalOnAddThreadTime( 0.0f, 0 );
+			PROFILER_SCOPE_LOG_TIME( TEXT( "5     FProfilerSession::HandleTicker.OnAddThreadTime" ), &Current );
+
+			// Update mini-view.
+			OnAddThreadTime.ExecuteIfBound( LastFrameIndex, ThreadMS, StatMetaData );
+		}
 
 		FrameToProfilerDataMapping.Remove( FrameIndex );
 	}	
@@ -384,6 +423,12 @@ bool FProfilerSession::HandleTicker( float DeltaTime )
 	{
 		if( FrameToProcess.Num() == 0 && bHasAllProfilerData )
 		{
+			CompletionSyncAggregatedEventGraphData();
+
+			// Advance event graphs.
+			EventGraphDataMaximum->Advance( 0, DataProvider->GetNumFrames() );
+			EventGraphDataTotal->Advance( 0, DataProvider->GetNumFrames() );
+
 			// Broadcast that a capture file has been fully processed.
 			OnCaptureFileProcessed.ExecuteIfBound( GetInstanceID() );
 
@@ -458,6 +503,15 @@ void FProfilerSession::LoadComplete()
 	bHasAllProfilerData = true;
 }
 
+float FProfilerSession::GetProgress() const
+{
+	if (NumFrames > 0)
+	{
+		return (float)NumFramesProcessed / NumFrames;
+	}
+	return 0.0f;
+}
+
 void FProfilerSession::UpdateProfilerData( const FProfilerDataFrame& Content )
 {
 	FrameToProfilerDataMapping.FindOrAdd( Content.Frame ) = Content;
@@ -466,13 +520,12 @@ void FProfilerSession::UpdateProfilerData( const FProfilerDataFrame& Content )
 
 void FProfilerSession::UpdateMetadata( const FStatMetaData& InClientStatMetaData )
 {
-	const uint32 CurrentStatMetaDataSize = InClientStatMetaData.GetMetaDataSize();
-	if( CurrentStatMetaDataSize != StatMetaDataSize )
+	const uint32 NewStatMetaDataSize = InClientStatMetaData.GetMetaDataSize();
+	if( NewStatMetaDataSize != StatMetaDataSize )
 	{
-		ClientStatMetadata = &InClientStatMetaData;
+		ClientStatMetadata = InClientStatMetaData;
 		bRequestStatMetadataUpdate = true;
-
-		StatMetaDataSize = CurrentStatMetaDataSize;
+		StatMetaDataSize = NewStatMetaDataSize;
 	}
 }
 
@@ -534,6 +587,7 @@ bool FRawProfilerSession::HandleTicker( float DeltaTime )
 static double GetSecondsPerCycle( const FStatPacketArray& Frame )
 {
 	const FName SecondsPerCycleFName = FName(TEXT("//STATGROUP_Engine//STAT_SecondsPerCycle///Seconds$32$Per$32$Cycle///////STATCAT_Advanced////"));
+	const FName SecondsPerCycleRawName = FStatConstants::RAW_SecondsPerCycle;
 	double Result = 0;
 
 	for( int32 PacketIndex = 0; PacketIndex < Frame.Packets.Num(); PacketIndex++ )
@@ -846,14 +900,8 @@ void FProfilerStatMetaData::Update( const FStatMetaData& ClientStatMetaData )
 {
 	PROFILER_SCOPE_LOG_TIME( TEXT( "FProfilerStatMetaData.Update" ), nullptr );
 
-	FStatMetaData LocalCopy;
-	{
-		FScopeLock Lock( ClientStatMetaData.CriticalSection );
-		LocalCopy = ClientStatMetaData;
-	}
-
 	// Iterate through all thread descriptions.
-	ThreadDescriptions.Append( LocalCopy.ThreadDescriptions );
+	ThreadDescriptions.Append( ClientStatMetaData.ThreadDescriptions );
 
 	// Initialize fake stat for Self.
 	const uint32 NoGroupID = 0;
@@ -863,20 +911,20 @@ void FProfilerStatMetaData::Update( const FStatMetaData& ClientStatMetaData )
 	InitializeStat( 1, NoGroupID, FStatConstants::NAME_ThreadRoot.GetPlainNameString(), STATTYPE_CycleCounter, FStatConstants::NAME_ThreadRoot );
 
 	// Iterate through all stat group descriptions.
-	for( auto It = LocalCopy.GroupDescriptions.CreateConstIterator(); It; ++It )
+	for( auto It = ClientStatMetaData.GroupDescriptions.CreateConstIterator(); It; ++It )
 	{
 		const FStatGroupDescription& GroupDesc = It.Value();
 		InitializeGroup( GroupDesc.ID, GroupDesc.Name );
 	}
 
 	// Iterate through all stat descriptions.
-	for( auto It = LocalCopy.StatDescriptions.CreateConstIterator(); It; ++It )
+	for( auto It = ClientStatMetaData.StatDescriptions.CreateConstIterator(); It; ++It )
 	{
 		const FStatDescription& StatDesc = It.Value();
 		InitializeStat( StatDesc.ID, StatDesc.GroupID, StatDesc.Name, (EStatType)StatDesc.StatType );
 	}
 
-	SecondsPerCycle = LocalCopy.SecondsPerCycle;
+	SecondsPerCycle = ClientStatMetaData.SecondsPerCycle;
 }
 
 
diff --git a/Engine/Source/Developer/Profiler/Private/ProfilerSession.h b/Engine/Source/Developer/Profiler/Private/ProfilerSession.h
--- a/Engine/Source/Developer/Profiler/Private/ProfilerSession.h
+++ b/Engine/Source/Developer/Profiler/Private/ProfilerSession.h
@@ -26,6 +26,9 @@ namespace EProfilerSessionTypes
 
 		/** Invalid enum type, may be used as a number of enumerations. */
 		InvalidOrMax,
+
+
+
 	};
 
 	/**
@@ -1068,10 +1071,10 @@ public:
 
 	const FEventGraphDataRef GetEventGraphDataAverage() const
 	{
-		const uint32 NumFrames = EventGraphDataTotal->GetNumFrames();
+		const uint32 NumFramesLocal = EventGraphDataTotal->GetNumFrames();
 		FEventGraphDataRef EventGraphDataAverage = EventGraphDataTotal->DuplicateAsRef();
-		EventGraphDataAverage->Divide( (double)NumFrames );
-		EventGraphDataAverage->Advance( 0, NumFrames );
+		EventGraphDataAverage->Divide( (double)NumFramesLocal );
+		EventGraphDataAverage->Advance( 0, NumFramesLocal );
 
 		return EventGraphDataAverage;
 	}
@@ -1139,13 +1142,27 @@ protected:
 	 */
 	void UpdateAggregatedEventGraphData( const uint32 FrameIndex );
 
-	void EventGraphCombineAndMax( const FEventGraphDataRef Current, const uint32 NumFrames );
+	/** Completion sync. */
+	void CompletionSyncAggregatedEventGraphData();
 
-	void EventGraphCombineAndAdd( const FEventGraphDataRef Current, const uint32 NumFrames );
+	void EventGraphCombineAndMax( const FEventGraphData* Current, const uint32 NumFrames );
+
+	void EventGraphCombineAndAdd( const FEventGraphData* Current, const uint32 NumFrames );
 
 	/** Called when the capture file has been fully loaded. */
 	void LoadComplete();
 
+	/** Sets number of frames. */
+	void SetNumberOfFrames( int32 InNumFrames )
+	{
+		NumFrames = InNumFrames;
+	}
+	
+	/**
+	 * @return progress as floating point between 0 and 1.
+	 */
+	float GetProgress() const;
+
 protected:
 	/** All aggregated stats, stored as StatID -> FProfilerAggregatedStat. */
 	TMap<uint32, FProfilerAggregatedStat > AggregatedStats;
@@ -1158,8 +1175,8 @@ protected:
 	TMap<uint32, FProfilerDataFrame> FrameToProfilerDataMapping;
 	/** Frame indices that should be processed by the profiler manager. */
 	TArray<uint32> FrameToProcess;
-	/** Reference to the client's stat metadata. */
-	const FStatMetaData* ClientStatMetadata;
+	/** Copy of the client stats metadata. */
+	FStatMetaData ClientStatMetadata;
 	/** If true, we need to update the metadata before we update the data provider. */
 	bool bRequestStatMetadataUpdate;
 	bool bLastPacket;
@@ -1186,8 +1203,8 @@ protected:
 	/** Aggregated event graph data for all collected frames, contains only maximum inclusive times. */
 	FEventGraphDataRef EventGraphDataMaximum;
 
-	/** Temporary event graph data for the specified frame. */
-	FEventGraphDataRef EventGraphDataCurrent;
+	/** Temporary event graph data for the specified frame. Recreated each frame, used by the task graph tasks only. */
+	const FEventGraphData* EventGraphDataCurrent;
 
 	/** Event graph completion sync ( combine for max, combine for add ) that can be done in parallel, but need to wait before combining the next frame. */
 	FGraphEventRef CompletionSync;
@@ -1208,6 +1225,12 @@ protected:
 	/** Data filepath. */
 	/*const*/ FString DataFilepath;
 
+	/** Number of frames in the file. */	
+	int32 NumFrames;
+
+	/** Number of frames already processed. */
+	int32 NumFramesProcessed;
+
 	/** True, if this profiler session instance is currently previewing data, only valid if profiler is connected to network based session. */
 	bool bDataPreviewing;
 
