Copying //UE4/Dev-Core to //UE4/Main

diff --git a/Engine/Source/Developer/Profiler/Private/ProfilerDataSource.h b/Engine/Source/Developer/Profiler/Private/ProfilerDataSource.h
--- a/Engine/Source/Developer/Profiler/Private/ProfilerDataSource.h
+++ b/Engine/Source/Developer/Profiler/Private/ProfilerDataSource.h
@@ -87,15 +87,13 @@ class FGraphDataSourceDescription
 public:
 	FGraphDataSourceDescription( const uint32 InStatID )
 		: StatID( InStatID )
-		, GroupID( -1 )
 		, SampleType( EProfilerSampleTypes::InvalidOrMax )
 		, CreationTime( -1 )
 	{}
 
-	void Initialize( const FString InStatName, const uint32 InGroupID, const FString InGroupName, const EProfilerSampleTypes::Type InSampleType, const FDateTime InCreationTime )
+	void Initialize( const FString InStatName, const FString InGroupName, const EProfilerSampleTypes::Type InSampleType, const FDateTime InCreationTime )
 	{
 		StatName     = InStatName;
-		GroupID      = InGroupID;
 		GroupName    = InGroupName;
 		SampleType   = InSampleType;
 		CreationTime = InCreationTime;
@@ -109,14 +107,6 @@ public:
 		return StatID;
 	}
 
-	/**
-	 * @return the ID of the stat group owned by this data source.
-	 */
-	const uint32 GetGroupID() const
-	{
-		return GroupID;
-	}
-
 	/**
 	 * @return name of the stat owned by this data source.
 	 */
@@ -159,9 +149,6 @@ protected:
 	/** The ID of the stat owned by this data source. */
 	const uint32 StatID;
 
-	/** The ID of the stat group owned by this data source. */
-	uint32 GroupID;
-
 	/** The name of the stat owned by this data source. */
 	FString StatName;
 
@@ -528,13 +515,6 @@ public:
 	 */
 	const FProfilerAggregatedStat* GetAggregatedStat() const;
 
-
-	// @TODO: Returns an offset based on the creation time, used to synchronize several graphs with different creation times.
-	const float GetOffsetMS() const
-	{
-		return 0.0f;
-	}
-
 	const bool CanBeDisplayedAsTimeBased() const
 	{
 		return true;
@@ -560,8 +540,6 @@ protected:
 };
 
 
-class FGraphDataSourceVerticalAggregate;
-
 /** 
  * A specialized view of a few data providers. Provides access only to the specified group of data. 
  * Data is interpolated for 60 frames per second.
@@ -716,13 +694,17 @@ namespace EEventPropertyIndex
 		StatName,
 		InclusiveTimeMS,
 		InclusiveTimePct,
-		MinInclusiveTimeMS,
-		MaxInclusiveTimeMS,
-		AvgInclusiveTimeMS,
 		ExclusiveTimeMS,
 		ExclusiveTimePct,
-		AvgInclusiveTimePerCallMS,
 		NumCallsPerFrame,
+		/** Special name used for unknown property. */
+		None,
+
+		MinInclusiveTimeMS,
+		MaxInclusiveTimeMS,
+		AvgInclusiveTimeMS,	
+		MinNumCallsPerFrame,
+		MaxNumCallsPerFrame,
 		AvgNumCallsPerFrame,
 		ThreadName,
 		ThreadDurationMS,
@@ -732,9 +714,7 @@ namespace EEventPropertyIndex
 		ThreadToFramePct,
 		StartTimeMS,
 		GroupName,
-		/** Special name used for unknown property. */
-		None,
-
+		
 		// Booleans
 		bIsHotPath,
 		bIsFiltered,
@@ -979,12 +959,12 @@ public:
 	/** Initializes the minimal property manager for the event graph sample. */
 	static void InitializePropertyManagement();
 
-	static const FEventProperty& GetEventPropertyByIndex( const EEventPropertyIndex::Type PropertyIndex )
+	static FORCEINLINE_DEBUGGABLE const FEventProperty& GetEventPropertyByIndex( const EEventPropertyIndex::Type PropertyIndex )
 	{
 		return Properties[PropertyIndex];
 	}
 
-	static const FEventProperty& GetEventPropertyByName( const FName PropertyName )
+	static FORCEINLINE_DEBUGGABLE const FEventProperty& GetEventPropertyByName( const FName PropertyName )
 	{
 		return *NamedProperties.FindChecked( PropertyName );
 	}
@@ -997,7 +977,7 @@ protected:
 	static TMap<FName,const FEventProperty*> NamedProperties;
 
 public:
-	/** Initialization constructor. */
+	/** Initialization constructor. Only used for root event. */
 	FEventGraphSample( const FName InName )
 		: _ParentPtr( nullptr )
 		, _RootPtr( nullptr )
@@ -1006,22 +986,23 @@ public:
 		, _GroupName( InName )
 		, _StatName( InName )
 		, _StatID( 0 )
-		, _StartTimeMS( 0.0f )
-		, _InclusiveTimeMS( 0.0f )
-		, _InclusiveTimePct( 0.0f )
-		, _MinInclusiveTimeMS( 0.0f )
-		, _MaxInclusiveTimeMS( 0.0f )
-		, _AvgInclusiveTimeMS( 0.0f )
-		, _AvgInclusiveTimePerCallMS( 0.0000f )
-		, _NumCallsPerFrame( 0.0f )
-		, _AvgNumCallsPerFrame( 0.0f )
-		, _ExclusiveTimeMS( 0.0f )
-		, _ExclusiveTimePct( 0.0f )
-		, _FrameDurationMS( 0.0f )
-		, _ThreadDurationMS( 0.0f )
-		, _ThreadToFramePct( 0.0f )
-		, _ThreadPct( 0.0f )
-		, _FramePct( 0.0f )
+		, _StartTimeMS( 0.0 )
+		, _InclusiveTimeMS( 0.0 )
+		, _InclusiveTimePct( 0.0 )
+		, _MinInclusiveTimeMS( TNumericLimits<double>::Max() )
+		, _MaxInclusiveTimeMS( TNumericLimits<double>::Min() )
+		, _AvgInclusiveTimeMS( 0.0 )
+		, _NumCallsPerFrame( 1.0 )
+		, _MinNumCallsPerFrame( TNumericLimits<double>::Max() )
+		, _MaxNumCallsPerFrame( TNumericLimits<double>::Min() )
+		, _AvgNumCallsPerFrame( 0.0 )
+		, _ExclusiveTimeMS( 0.0 )
+		, _ExclusiveTimePct( 0.0 )
+		, _FrameDurationMS( 0.0 )
+		, _ThreadDurationMS( 0.0 )
+		, _ThreadToFramePct( 0.0 )
+		, _ThreadPct( 0.0 )
+		, _FramePct( 0.0 )
 		, _bIsHotPath( false )
 		, _bIsFiltered( false )
 		, _bIsCulled( false )
@@ -1042,26 +1023,20 @@ protected:
 	FEventGraphSample
 	(
 		const FName InThreadName,
- 
+
 		const FName InGroupName,
 		const uint32 InStatID,
 		const FName InStatName,
- 
+
 		const double InStartTimeMS,
 		const double InInclusiveTimeMS,
- 
-		const double InMinInclusiveTimeMS,
-		const double InMaxInclusiveTimeMS,
-		const double InAvgInclusiveTimeMS,
- 
-		const double InAvgPerCallInclusiveTimeMS,
- 
 		const double InNumCallsPerFrame,
-		const double InAvgNumCallsPerFrame,
- 
+
 		const FEventGraphSamplePtr InParentPtr = NULL
 	)
 		: _ParentPtr( InParentPtr )
+		, _RootPtr( nullptr )
+		, _ThreadPtr( nullptr )
 		, _ThreadName( InThreadName )
 		, _GroupName( InGroupName )
 		, _StatName( InStatName )
@@ -1069,12 +1044,13 @@ protected:
 		, _StartTimeMS( InStartTimeMS )
 		, _InclusiveTimeMS( InInclusiveTimeMS )
 		, _InclusiveTimePct( 0.0f )
-		, _MinInclusiveTimeMS( InMinInclusiveTimeMS )
-		, _MaxInclusiveTimeMS( InMaxInclusiveTimeMS )
-		, _AvgInclusiveTimeMS( InAvgInclusiveTimeMS )
-		, _AvgInclusiveTimePerCallMS( 0.0000f )
+		, _MinInclusiveTimeMS( InInclusiveTimeMS )
+		, _MaxInclusiveTimeMS( InInclusiveTimeMS )
+		, _AvgInclusiveTimeMS( InInclusiveTimeMS )
 		, _NumCallsPerFrame( InNumCallsPerFrame )
-		, _AvgNumCallsPerFrame( InAvgNumCallsPerFrame )
+		, _MinNumCallsPerFrame( InNumCallsPerFrame )
+		, _MaxNumCallsPerFrame( InNumCallsPerFrame )
+		, _AvgNumCallsPerFrame( InNumCallsPerFrame )
 		, _ExclusiveTimeMS( 0.0f )
 		, _ExclusiveTimePct( 0.0f )
 		, _FrameDurationMS( 0.0f )
@@ -1103,8 +1079,9 @@ protected:
 		, _MinInclusiveTimeMS( SourceEvent._MinInclusiveTimeMS )
 		, _MaxInclusiveTimeMS( SourceEvent._MaxInclusiveTimeMS )
 		, _AvgInclusiveTimeMS( SourceEvent._AvgInclusiveTimeMS )
-		, _AvgInclusiveTimePerCallMS( SourceEvent._AvgInclusiveTimePerCallMS )
 		, _NumCallsPerFrame( SourceEvent._NumCallsPerFrame )
+		, _MinNumCallsPerFrame( SourceEvent._MinNumCallsPerFrame )
+		, _MaxNumCallsPerFrame( SourceEvent._MaxNumCallsPerFrame )
 		, _AvgNumCallsPerFrame( SourceEvent._AvgNumCallsPerFrame )
 		, _ExclusiveTimeMS( SourceEvent._ExclusiveTimeMS )
 		, _ExclusiveTimePct( SourceEvent._ExclusiveTimePct )
@@ -1120,114 +1097,123 @@ protected:
 	{}
 
 public:
+	/** Destructor. */
+	~FEventGraphSample()
+	{
+#if	_DEBUG
+		int k = 0; k++;
+#endif // _DEBUG
+	}
+
+
 	/*-----------------------------------------------------------------------------
 		Operations
 	-----------------------------------------------------------------------------*/
 
-	void AddSamplePtr( const FEventGraphSamplePtr& Other )
+	FORCEINLINE_DEBUGGABLE void Combine( const FEventGraphSamplePtr& Other )
 	{
-		(*this) += *Other.Get();
-	}
+		// Total
+		_InclusiveTimeMS += Other->_InclusiveTimeMS;
+		_NumCallsPerFrame += Other->_NumCallsPerFrame;
+		//_ExclusiveTimeMS += Other->_ExclusiveTimeMS;
 
-	void DivideSamplePtr( const double Divisor )
-	{
-		(*this) /= Divisor;
-	}
+		// Min/Max
+		_MinInclusiveTimeMS = FMath::Min( _MinInclusiveTimeMS, Other->_InclusiveTimeMS );
+		_MaxInclusiveTimeMS = FMath::Max( _MaxInclusiveTimeMS, Other->_InclusiveTimeMS );
 
-	void MaxSamplePtr( const FEventGraphSamplePtr& Other )
-	{
-		Max( *Other.Get() );
+		_MinNumCallsPerFrame = FMath::Min( _MinNumCallsPerFrame, Other->_NumCallsPerFrame );
+		_MaxNumCallsPerFrame = FMath::Max( _MaxNumCallsPerFrame, Other->_NumCallsPerFrame );
 	}
 
-	FEventGraphSample& operator+=( const FEventGraphSample& Other )
+protected:
+
+	/** For creating per-frame average event graph. */
+	FORCEINLINE_DEBUGGABLE void CopyAverage( const double NumFrames )
 	{
-		_InclusiveTimeMS += Other._InclusiveTimeMS;
-		_MinInclusiveTimeMS += Other._MinInclusiveTimeMS;
-		_MaxInclusiveTimeMS += Other._MaxInclusiveTimeMS;
-		_AvgInclusiveTimeMS += Other._AvgInclusiveTimeMS;
-		_AvgInclusiveTimePerCallMS += Other._AvgInclusiveTimePerCallMS;
-		_NumCallsPerFrame += Other._NumCallsPerFrame;
-		_AvgNumCallsPerFrame += Other._AvgNumCallsPerFrame;
-		_ExclusiveTimeMS += Other._ExclusiveTimeMS;
+		_InclusiveTimeMS = _AvgInclusiveTimeMS;
+		_NumCallsPerFrame = _AvgNumCallsPerFrame;
 
-		return *this;
+		_FrameDurationMS /= NumFrames;
+		_ThreadDurationMS /= NumFrames;
+
+		FixFrameThreadPcts();
 	}
 
-	FEventGraphSample& operator/=( const double Divisor )
+	/** For creating highest "per-frame" event graph. */
+	FORCEINLINE_DEBUGGABLE void CopyMaximum( FEventGraphSample* RootEvent, FEventGraphSample* ThreadEvent )
 	{
-		_InclusiveTimeMS /= Divisor;
-		_MinInclusiveTimeMS /= Divisor;
-		_MaxInclusiveTimeMS /= Divisor;
-		_AvgInclusiveTimeMS /= Divisor;
-		_AvgInclusiveTimePerCallMS /= Divisor;
-		_NumCallsPerFrame /= Divisor;
-		_AvgNumCallsPerFrame /= Divisor;
-		_ExclusiveTimeMS /= Divisor;
+		_InclusiveTimeMS = _MaxInclusiveTimeMS;
+		_NumCallsPerFrame = _MaxNumCallsPerFrame;
 
-		return *this;
-	}
+		_ThreadDurationMS = ThreadEvent->_MaxInclusiveTimeMS;
+		_FrameDurationMS = RootEvent->_MaxInclusiveTimeMS;
 
-	void Max( const FEventGraphSample& Other )
-	{
-		_InclusiveTimeMS = FMath::Max( _InclusiveTimeMS, Other._InclusiveTimeMS );
-		_MinInclusiveTimeMS = FMath::Max( _MinInclusiveTimeMS, Other._MinInclusiveTimeMS );
-		_MaxInclusiveTimeMS = FMath::Max( _MaxInclusiveTimeMS, Other._MaxInclusiveTimeMS );
-		_AvgInclusiveTimeMS = FMath::Max( _AvgInclusiveTimeMS, Other._AvgInclusiveTimeMS );
-		_AvgInclusiveTimePerCallMS = FMath::Max( _AvgInclusiveTimePerCallMS, Other._AvgInclusiveTimePerCallMS );
-		_NumCallsPerFrame = FMath::Max( _NumCallsPerFrame, Other._NumCallsPerFrame );
-		_AvgNumCallsPerFrame = FMath::Max( _AvgNumCallsPerFrame, Other._AvgNumCallsPerFrame );
-		_ExclusiveTimeMS = FMath::Max( _ExclusiveTimeMS, Other._ExclusiveTimeMS );
-	}
+		// Exclusive values don't make sense for max.
+		_ExclusiveTimeMS = 0.0;
+		_ExclusiveTimePct = 0.0;
 
-	FORCEINLINE_DEBUGGABLE  bool AreTheSamePtr( const FEventGraphSamplePtr& Other ) const
-	{
-		return *this == *Other.Get();
+		FixFrameThreadPcts();
 	}
 
-	const bool operator==( const FEventGraphSample& Other ) const
+	/** Calculates time and percentage values that may depend on a child's parent. */
+	void FixChildrenTimesAndCalcAverages( const double NumFrames )
 	{
-		const bool bAreChildrenTheSame = _ThreadName==Other._ThreadName && _StatID==Other._StatID;// && _GroupName==Other._GroupName;
-		return bAreChildrenTheSame;
-	}
+		if (_StatName != FEventGraphConsts::RootEvent)
+		{
+			// Get correct value for frame and thread duration.
+			_FrameDurationMS = GetRoot()->_InclusiveTimeMS;
+			_ThreadDurationMS = GetThread()->_InclusiveTimeMS;
 
-	/** Destructor. */
-	~FEventGraphSample()
-	{
-#if	_DEBUG
-		int k=0;k++;
-#endif // _DEBUG
-	}
+			FEventGraphSample* Parent = GetParent().Get();
+			if (IsSelf())
+			{
+				Parent->_ExclusiveTimeMS = _InclusiveTimeMS;
+				Parent->_ExclusiveTimePct = 100.0f * Parent->_ExclusiveTimeMS / Parent->_InclusiveTimeMS;
+			}
 
-	FEventGraphSamplePtr FindChildPtr( const FEventGraphSamplePtr& ChildToLook );
-	// TODO: CombineAndAdd, CombineAndFindMax, Divide FEventGraphDataRef<AggregateFunc,FinalizerFunc,bRequiredCopy>
-	void CombineAndAddPtr_Recurrent( const FEventGraphSamplePtr& Other );
-	void CombineAndFindMaxPtr_Recurrent( const FEventGraphSamplePtr& Other );
-	void Divide_Recurrent( const double Divisor );
+			_InclusiveTimePct = 100.0f * _InclusiveTimeMS / Parent->_InclusiveTimeMS;
 
-	struct FDivide
-	{
-		FORCEINLINE void operator()( FEventGraphSample* EventPtr, const double Divisor )
+			FixFrameThreadPcts();
+		}
+		else
 		{
-			*EventPtr /= Divisor;
+			_InclusiveTimePct = 100.0f;
+			_ThreadDurationMS = _InclusiveTimeMS;
+			_FrameDurationMS = _InclusiveTimeMS;
 		}
-	};
 
-	/**
-	 * @return true, if this event is a root event.
-	 */
-	bool IsRoot() const
+		_AvgInclusiveTimeMS = _InclusiveTimeMS / NumFrames;
+		_AvgNumCallsPerFrame = _NumCallsPerFrame / NumFrames;
+	}
+
+
+	void FixFrameThreadPcts()
 	{
-		return _StatName == FEventGraphConsts::RootEvent;
+		_ThreadToFramePct = 100.0f*_ThreadDurationMS / _FrameDurationMS;
+		_ThreadPct = 100.0f*_InclusiveTimeMS / _ThreadDurationMS;
+		_FramePct = 100.0f*_InclusiveTimeMS / _FrameDurationMS;
 	}
 
-	/**
-	 * @return true, if this event is a fake self event.
-	 */
-	bool IsSelf() const
+	/** Copy maximum values for all children, also fixes thread/frame time. */
+	void SetMaximumTimesForAllChildren();
+
+	/** Sets root and thread for all children. */
+	void SetRootAndThreadForAllChildren();
+
+	/** Fixes children times and calculates average values. */
+	void FixChildrenTimesAndCalcAveragesForAllChildren( const double NumFrames );
+
+	FORCEINLINE_DEBUGGABLE  bool AreTheSamePtr( const FEventGraphSamplePtr& Other ) const
 	{
-		return _StatName == FEventGraphConsts::Self;
+		const bool bAreChildrenTheSame = _ThreadName == Other->_ThreadName && _StatID == Other->_StatID;// && _GroupName==Other._GroupName;
+		return bAreChildrenTheSame;
 	}
 
+
+	FEventGraphSamplePtr FindChildPtr( const FEventGraphSamplePtr& ChildToLook );
+	void Combine_Recurrent( const FEventGraphSamplePtr& Other );
+
+public:
 	template< typename TFunc > 
 	FORCEINLINE_DEBUGGABLE void ExecuteOperationForAllChildren( TFunc FuncToCall )
 	{
@@ -1258,7 +1244,7 @@ public:
 	}
 
 	template< typename TFunc, typename TArg0 > 
-	FORCEINLINE_DEBUGGABLE void ExecuteOperationForAllChildren( TFunc FuncToCall, TArg0& Arg0 )
+	FORCEINLINE_DEBUGGABLE void ExecuteOperationForAllChildren( TFunc FuncToCall, TArg0 Arg0 )
 	{
 		const bool bAllowShrinking = false;
 		TArray<FEventGraphSample*>& Stack = FProfilerScratchArea::Get().ExecuteOperationArray;
@@ -1286,6 +1272,35 @@ public:
 		Stack.Reset();
 	}
 
+	template< typename TFunc, typename TArg0, typename TArg1 >
+	FORCEINLINE_DEBUGGABLE void ExecuteOperationForAllChildren( TFunc FuncToCall, TArg0 Arg0, TArg1 Arg1 )
+	{
+		const bool bAllowShrinking = false;
+		TArray<FEventGraphSample*>& Stack = FProfilerScratchArea::Get().ExecuteOperationArray;
+
+		Stack.Add( &AsShared().Get() );
+		int32 Idx = 1;
+
+		while (Stack.Num() > 0)
+		{
+			// Get the parent and assign events.
+			FEventGraphSample* Current = Stack.Pop( bAllowShrinking );
+			FuncToCall( Current, Arg0, Arg1 );
+
+			// Push children onto the stack.
+			const TArray<FEventGraphSamplePtr>& ChildrenPtr = Current->GetChildren();
+			const int32 NumChildren = ChildrenPtr.Num();
+			for (int32 ChildIndex = 0; ChildIndex < NumChildren; ++ChildIndex)
+			{
+				const FEventGraphSamplePtr& ChildPtr = ChildrenPtr[ChildIndex];
+				Stack.Push( ChildPtr.Get() );
+			}
+		}
+
+		// Reset, but keep the allocation.
+		Stack.Reset();
+	}
+
 protected:
 	/** @return a shared pointer to the newly created copy of this event graph sample, creates a full copy of hierarchy and duplicates all samples. */
 	FEventGraphSamplePtr DuplicateWithHierarchyPtr()
@@ -1313,23 +1328,7 @@ protected:
 	}
 
 
-	void SetRootAndThreadForAllChildren()
-	{
-		FEventGraphSamplePtr RootEvent = AsShared();
-		const int32 NumChildren = _ChildrenPtr.Num();
-		for( int32 ChildIndex = 0; ChildIndex < NumChildren; ++ChildIndex )
-		{
-			const FEventGraphSamplePtr& ThreadEvent = _ChildrenPtr[ChildIndex];
-			ThreadEvent->SetRootAndThreadEvents_Iterative( RootEvent.Get(), ThreadEvent.Get() );
-		}
-	}
-
-	FORCEINLINE void SetRootAndThread( FEventGraphSample* RootEvent, FEventGraphSample* ThreadEvent )
-	{
-		_RootPtr = RootEvent;
-		_ThreadPtr = ThreadEvent;
-	}
-
+#if	0
 	FORCEINLINE_DEBUGGABLE void SetRootAndThreadEvents_Iterative( FEventGraphSample* RootEvent, FEventGraphSample* ThreadEvent )
 	{
 		const bool bAllowShrinking = false;
@@ -1342,7 +1341,7 @@ protected:
 		{
 			// Get the parent and assign events.
 			FEventGraphSample* Current = Stack.Pop( bAllowShrinking );
-			Current->SetRootAndThread( RootEvent, ThreadEvent );
+			FSetRootAndThread()(Current, RootEvent, ThreadEvent);
 
 			// Push children onto the stack.
 			const TArray<FEventGraphSamplePtr>& ChildrenPtr = Current->GetChildren();
@@ -1357,6 +1356,7 @@ protected:
 		// Reset, but keep the allocation.
 		Stack.Reset();
 	}
+#endif // 0
 
 public:
 	/** @return a shared pointer to the newly created copy of this event graph sample, without any children and with no parent. */
@@ -1365,13 +1365,28 @@ public:
 		return MakeShareable( new FEventGraphSample( *this, FDuplicateSimpleTag() ) );
 	}
 
-public:
 	/** Adds a child to this sample. */
 	FORCEINLINE void AddChildPtr( const FEventGraphSamplePtr& Child ) 
 	{
 		_ChildrenPtr.Add( Child );
 	}
 
+	/**
+	 * @return true, if this event is a root event.
+	 */
+	bool IsRoot() const
+	{
+		return _StatName == FEventGraphConsts::RootEvent;
+	}
+
+	/**
+	 * @return true, if this event is a fake self event.
+	 */
+	bool IsSelf() const
+	{
+		return _StatName == FEventGraphConsts::Self;
+	}
+
 	/**
 	 * @return a shared pointer to the parent of this event, may be null.
 	 */
@@ -1529,6 +1544,7 @@ private:
 	}
 	
 public:
+#if	0
 	/**
 	 * @return children time in ms, excluding the self time by default
 	 */
@@ -1554,7 +1570,6 @@ public:
 		return ChildrenTotalTimeMS;
 	}
 
-
 	/**
 	 * @return if exist, a fake self event for this event sample, otherwise null.
 	 */
@@ -1574,51 +1589,7 @@ public:
 		return Self;
 	}
 
-	/** Not used, optimized version @see ExecuteOperationForAllChildren. */
-	void FixAllChildren_Recurrent()
-	{
-		FixChildrenTimes();
-
-		// Check other children.
-		const int32 NumChildren = _ChildrenPtr.Num();
-		for( int32 ChildIndex = 0; ChildIndex < NumChildren; ++ChildIndex )
-		{
-			FEventGraphSamplePtr ThisChild = _ChildrenPtr[ChildIndex];
-			ThisChild->FixAllChildren_Recurrent();
-		}
-	}
-
-	struct FFixChildrenTimes
-	{
-		FORCEINLINE void operator()( FEventGraphSample* EventPtr )
-		{
-			EventPtr->FixChildrenTimes();
-		}
-	};
-
-	/** Calculates time and percentage values that may depend on a child's parent. */
-	void FixChildrenTimes()
-	{
-		if( _StatName != FEventGraphConsts::RootEvent )
-		{
-			// Get correct value for frame and thread duration.
-			_FrameDurationMS = GetRoot()->_InclusiveTimeMS;
-			_ThreadDurationMS = GetThread()->_InclusiveTimeMS;
-
-			FEventGraphSample* Parent = GetParent().Get();
-			if( IsSelf() )
-			{
-				Parent->_ExclusiveTimeMS = _InclusiveTimeMS;
-				Parent->_ExclusiveTimePct = 100.0f * Parent->_ExclusiveTimeMS/Parent->_InclusiveTimeMS;
-			}
-
-			_InclusiveTimePct = 100.0f * _InclusiveTimeMS / Parent->_InclusiveTimeMS;
-
-			_ThreadToFramePct = 100.0f*_ThreadDurationMS/_FrameDurationMS;
-			_ThreadPct = 100.0f*_InclusiveTimeMS/_ThreadDurationMS;
-			_FramePct = 100.0f*_InclusiveTimeMS/_FrameDurationMS;
-		}
-	}
+#endif // 0
 
 public:
 	double& PropertyValueAsDouble( const EEventPropertyIndex::Type PropertyIndex )
@@ -1723,31 +1694,6 @@ public:
 		}
 	}
  
-	// CLang -> explicit specialization of 'SetBooleanStateForAllChildren' in class scope
-// 	template<>
-// 	void SetBooleanStateForAllChildren<EEventPropertyIndex::bIsCulled>( const bool bValue )
-// 	{
-// 		ExecuteOperationForAllChildren( FSetBooleanState<&FEventGraphSample::_bIsCulled>(), bValue );
-// 	}
-//  
-// 	template<>
-// 	void SetBooleanStateForAllChildren<EEventPropertyIndex::bIsFiltered>( const bool bValue )
-// 	{
-// 		ExecuteOperationForAllChildren( FSetBooleanState<&FEventGraphSample::_bIsFiltered>(), bValue );
-// 	}
-//  
-// 	template<>
-// 	void SetBooleanStateForAllChildren<EEventPropertyIndex::bIsHotPath>( const bool bValue )
-// 	{
-// 		ExecuteOperationForAllChildren( FSetBooleanState<&FEventGraphSample::_bIsHotPath>(), bValue );
-// 	}
-//  
-// 	template<>
-// 	void SetBooleanStateForAllChildren<EEventPropertyIndex::bNeedNotCulledChildrenUpdate>( const bool bValue )
-// 	{
-// 		ExecuteOperationForAllChildren( FSetBooleanState<&FEventGraphSample::bNeedNotCulledChildrenUpdate>(), bValue );
-// 	}
-
 public:
 
 	/** Not used, optimized version @see SetBooleanStateForAllChildren. */
@@ -1802,31 +1748,26 @@ public:
 	/** Duration of this event and its children as percent of the caller. */
 	double _InclusiveTimePct;
 
-
-	// TODO: Rework to _MaxInclusiveTimeMS, _AvgInclusiveTimeMS
-
-	/** Minimum inclusive time of all instances for this event, in milliseconds. */ // TODO: to be removed, use data graph aggregation 
+	/** Minimum inclusive time of all instances for this event, in milliseconds. */ 
 	double _MinInclusiveTimeMS;
 
-	/** Maximum inclusive time of all instances for this event, in milliseconds. */ // TODO: to be removed, use data graph aggregation 
+	/** Maximum inclusive time of all instances for this event, in milliseconds. */
 	double _MaxInclusiveTimeMS;
 
-	/** Average inclusive time of all instances for this event, in milliseconds. */ // TODO: to be removed, use data graph aggregation 
+	/** Average inclusive time of all instances for this event, in milliseconds. */
 	double _AvgInclusiveTimeMS;
 
-	/** Average inclusive time per call of all instances for this event, in milliseconds. */
-	double _AvgInclusiveTimePerCallMS;
-
-	// TODO: Rename into _AvgNumCallsPerFrame and _MaxNumClassPerFrame
-
 	/** Number of times this event was called. */
 	double _NumCallsPerFrame;
 
-	/** Average number of times this event was called. */
-	double _AvgNumCallsPerFrame;	
+	/** Minimum number of times this event was called. */
+	double _MinNumCallsPerFrame;
 
 	/** Maximum number of times this event was called. */
-	//double _MaxNumCallsPerFrame;	
+	double _MaxNumCallsPerFrame;
+
+	/** Average number of times this event was called. */
+	double _AvgNumCallsPerFrame;	
 
 	/** Exclusive time of this event, in milliseconds. */
 	double _ExclusiveTimeMS;
@@ -2115,7 +2056,7 @@ private:
 -----------------------------------------------------------------------------*/
 
 /** Enumerates event graph types. */
-namespace EEventGraphTypes
+struct EEventGraphTypes
 {
 	enum Type
 	{
@@ -2128,6 +2069,9 @@ namespace EEventGraphTypes
 		/** Event graph for one frame, so both average and maximum can be used. */
 		OneFrame,
 
+		/** Selected frames event graph. */
+		Total,
+
 		/** Invalid enum type, may be used as a number of enumerations. */
 		InvalidOrMax,
 	};
@@ -2137,15 +2081,15 @@ namespace EEventGraphTypes
 	 *
 	 * @return string representation of the specified EEventGraphTypes value.
 	 */
-	FString ToName( const EEventGraphTypes::Type EventGraphType );
+	static FString ToName( const EEventGraphTypes::Type EventGraphType );
 
 	/**
 	 * @param EventGraphType - The value to get the string for.
 	 *
 	 * @return string representation with more detailed explanation of the specified EEventGraphTypes value.
 	 */
-	FString ToDescription( const EEventGraphTypes::Type EventGraphType );
-}
+	static FString ToDescription( const EEventGraphTypes::Type EventGraphType );
+};
 
 /** POD helper class for storing information like type and frame indices about the event graph. */
 class FEventGraphDataProperties
@@ -2240,14 +2184,17 @@ public:
 	 */
 	FEventGraphDataRef DuplicateAsRef();
 
-	void CombineAndAdd( const FEventGraphData& Other );
-
-	void CombineAndFindMax( const FEventGraphData& Other );
+	/** Combines current event graph with the second one. */
+	void Combine( const FEventGraphData& Other );
 
-	void Divide( const int32 NumFrames );
+	/** Updates data to get per-frame average event graph. */
+	void SetAsAverage();
 
+	/** Updates data to get highest "per-frame" event graph. */
+	void SetAsMaximim();
 
-	void Advance( const uint32 InFrameStartIndex, const uint32 InFrameEndIndex );
+	/** Finalizes current event graph. */
+	void Finalize( const uint32 InFrameStartIndex, const uint32 InFrameEndIndex );
 
 	/**
 	 * @return root event that contains all thread root events and its children
diff --git a/Engine/Source/Developer/Profiler/Private/ProfilerManager.cpp b/Engine/Source/Developer/Profiler/Private/ProfilerManager.cpp
--- a/Engine/Source/Developer/Profiler/Private/ProfilerManager.cpp
+++ b/Engine/Source/Developer/Profiler/Private/ProfilerManager.cpp
@@ -116,7 +116,6 @@ void FProfilerManager::PostConstructor()
 
 	ProfilerClient->OnProfilerFileTransfer().AddSP(this, &FProfilerManager::ProfilerClient_OnProfilerFileTransfer);
 
-	SessionManager->OnSelectedSessionChanged().AddSP( this, &FProfilerManager::SessionManager_OnSelectedSessionChanged );
 	SessionManager->OnInstanceSelectionChanged().AddSP( this, &FProfilerManager::SessionManager_OnInstanceSelectionChanged );
 
 	SetDataPreview( false );
@@ -425,11 +424,6 @@ void FProfilerManager::ProfilerClient_OnClientDisconnected( const FGuid& Session
 	SessionManager
 -----------------------------------------------------------------------------*/
 
-void FProfilerManager::SessionManager_OnSelectedSessionChanged( const ISessionInfoPtr& InActiveSession )
-{
-	//SessionManager_OnInstanceSelectionChanged(nullptr, false);
-}
-
 void FProfilerManager::SessionManager_OnInstanceSelectionChanged(const TSharedPtr<ISessionInstanceInfo>& InInstance, bool Selected)
 {
 	const ISessionInfoPtr& SelectedSession = SessionManager->GetSelectedSession();
@@ -472,16 +466,6 @@ void FProfilerManager::SessionManager_OnInstanceSelectionChanged(const TSharedPt
 	SessionInstancesUpdatedEvent.Broadcast();
 }
 
-/*-----------------------------------------------------------------------------
-	Misc
------------------------------------------------------------------------------*/
-
-FCombinedGraphDataSourceRef FProfilerManager::CreateCombinedGraphDataSource( const uint32 StatID )
-{
-	FCombinedGraphDataSource* CombinedGraphDataSource = new FCombinedGraphDataSource( StatID, FTimeAccuracy::FPS060 );
-	return MakeShareable( CombinedGraphDataSource );
-}
-
 /*-----------------------------------------------------------------------------
 	Stat tracking
 -----------------------------------------------------------------------------*/
@@ -494,57 +478,27 @@ bool FProfilerManager::TrackStat( const uint32 StatID )
 	const bool bStatIsReady = ProfilerSession->GetAggregatedStat( StatID ) != nullptr;
 	if( StatID != 0 && bStatIsReady )
 	{
-		FTrackedStat* TrackedStat = TrackedStats.Find( StatID );
-
-		if( TrackedStat == nullptr )
+		FTrackedStatPtr TrackedStat = TrackedStats.FindRef( StatID );
+		if (!TrackedStat.IsValid())
 		{
 			// R = H, G = S, B = V
-			const FLinearColor& ColorAverage = GetColorForStatID( StatID );
-			const FLinearColor ColorAverageHSV = ColorAverage.LinearRGBToHSV();
-
-			FLinearColor ColorBackgroundHSV = ColorAverageHSV;
-			ColorBackgroundHSV.G = FMath::Max( 0.0f, ColorBackgroundHSV.G-0.25f );
-
-			FLinearColor ColorExtremesHSV = ColorAverageHSV;
-			ColorExtremesHSV.G = FMath::Min( 1.0f, ColorExtremesHSV.G+0.25f );
-			ColorExtremesHSV.B = FMath::Min( 1.0f, ColorExtremesHSV.B+0.25f );
-
-			const FLinearColor ColorBackground = ColorBackgroundHSV.HSVToLinearRGB();
-			const FLinearColor ColorExtremes = ColorExtremesHSV.HSVToLinearRGB();
-
-			TrackedStat = &TrackedStats.Add( StatID, FTrackedStat(CreateCombinedGraphDataSource( StatID ),ColorAverage,ColorExtremes,ColorBackground,StatID) );
+			const FLinearColor& GraphColor = GetColorForStatID( StatID );
+// 			const FLinearColor ColorAverageHSV = ColorAverage.LinearRGBToHSV();
+// 
+// 			FLinearColor ColorBackgroundHSV = ColorAverageHSV;
+// 			ColorBackgroundHSV.G = FMath::Max( 0.0f, ColorBackgroundHSV.G-0.25f );
+// 
+// 			FLinearColor ColorExtremesHSV = ColorAverageHSV;
+// 			ColorExtremesHSV.G = FMath::Min( 1.0f, ColorExtremesHSV.G+0.25f );
+// 			ColorExtremesHSV.B = FMath::Min( 1.0f, ColorExtremesHSV.B+0.25f );
+// 
+// 			const FLinearColor ColorBackground = ColorBackgroundHSV.HSVToLinearRGB();
+// 			const FLinearColor ColorExtremes = ColorExtremesHSV.HSVToLinearRGB();
+
+			TrackedStat = TrackedStats.Add( StatID, MakeShareable( new FTrackedStat( ProfilerSession->CreateGraphDataSource( StatID ), GraphColor, StatID ) ) );
 			bAdded = true;
 
-			// @TODO: Convert a reference parameter to copy parameter/sharedptr/ref/weak, to avoid problems when a reference is no longer valid.
-			TrackedStatChangedEvent.Broadcast( *TrackedStat, true );
-		}
-
-		if( TrackedStat != nullptr )
-		{
-			uint32 NumAddedInstances = 0;
-			bool bMetadataInitialized = false;
-
-			if (ProfilerSession.IsValid())
-			{
-
-				const bool bInstanceAdded = TrackStatForSessionInstance( StatID, ActiveInstanceID );
-				NumAddedInstances += bInstanceAdded ? 1 : 0;
-
-				// Initialize metadata for combine graph data source.
-				if( !bMetadataInitialized )
-				{
-					const bool bIsStatReady = ProfilerSession->GetMetaData()->IsStatInitialized( StatID );
-					if( bIsStatReady )
-					{
-						const FProfilerStatMetaDataRef MetaData = ProfilerSession->GetMetaData();
-						const FProfilerStat& Stat = MetaData->GetStatByID( StatID );
-						const FProfilerGroup& Group = Stat.OwningGroup();
-
-						TrackedStat->CombinedGraphDataSource->Initialize( Stat.Name().GetPlainNameString(), Group.ID(), Group.Name().GetPlainNameString(), Stat.Type(), ProfilerSession->GetCreationTime() );
-						bMetadataInitialized = true;
-					}
-				}
-			}
+			TrackedStatChangedEvent.Broadcast( TrackedStat, true );
 		}
 	}
 
@@ -555,14 +509,19 @@ bool FProfilerManager::UntrackStat( const uint32 StatID )
 {
 	bool bRemoved = false;
 
-	const FTrackedStat* TrackedStat = TrackedStats.Find( StatID );
-	if( TrackedStat )
+	// Game thread time is always tracked.
+	const uint32 GameThreadStatID = ProfilerSession->GetMetaData()->GetGameThreadStatID();
+	if (StatID != GameThreadStatID)
 	{
-		TrackedStatChangedEvent.Broadcast( *TrackedStat, false );
-		TrackedStats.Remove( StatID );
-		bRemoved = true;
-	}
+		FTrackedStatPtr TrackedStat = TrackedStats.FindRef( StatID );
+		if (TrackedStat.IsValid())
+		{
+			TrackedStatChangedEvent.Broadcast( TrackedStat, false );
+			TrackedStats.Remove( StatID );
+			bRemoved = true;
+		}
 
+	}
 	return bRemoved;
 }
 
@@ -580,9 +539,10 @@ void FProfilerManager::ClearStatsAndInstances()
 
 	for( auto It = TrackedStats.CreateConstIterator(); It; ++It )
 	{
-		const FTrackedStat& TrackedStat = It.Value();
-		UntrackStat( TrackedStat.StatID );
+		const FTrackedStatPtr& TrackedStat = It.Value();
+		TrackedStatChangedEvent.Broadcast( TrackedStat, false );
 	}
+	TrackedStats.Empty();
 
 	ProfilerClient->Untrack( ActiveInstanceID );
 	ActiveInstanceID.Invalidate();
@@ -593,60 +553,6 @@ const bool FProfilerManager::IsStatTracked( const uint32 StatID ) const
 	return TrackedStats.Contains( StatID );
 }
 
-bool FProfilerManager::TrackStatForSessionInstance( const uint32 StatID, const FGuid& SessionInstanceID )
-{
-	bool bAdded = false;
-
-	const FTrackedStat* TrackedStat = TrackedStats.Find( StatID );
-	if( ProfilerSession.IsValid() && TrackedStat )
-	{
-		const bool bIsStatReady = ProfilerSession->GetMetaData()->IsStatInitialized( StatID );
-		const bool bIsExist = TrackedStat->CombinedGraphDataSource->IsProfilerSessionRegistered( SessionInstanceID );
-
-		if( bIsStatReady && !bIsExist )
-		{
-			// Create graph data provider for tracked stat.
-			FGraphDataSourceRefConst	GraphDataSource = ProfilerSession->CreateGraphDataSource( StatID );
-			TrackedStat->CombinedGraphDataSource->RegisterWithProfilerSession( SessionInstanceID, GraphDataSource );
-			bAdded = true;
-		}
-	}
-
-	return bAdded;
-}
-
-bool FProfilerManager::UntrackStatForSessionInstance( const uint32 StatID, const FGuid& SessionInstanceID )
-{
-	bool bRemoved = false;
-
-	const FTrackedStat* TrackedStat = TrackedStats.Find( StatID );
-
-	if( TrackedStat )
-	{
-		const bool bIsExist = TrackedStat->CombinedGraphDataSource->IsProfilerSessionRegistered( SessionInstanceID );
-		if( bIsExist )
-		{
-			TrackedStat->CombinedGraphDataSource->UnregisterWithProfilerSession( SessionInstanceID );
-			bRemoved = true;
-		}
-	}
-
-	return bRemoved;
-}
-
-const bool FProfilerManager::IsStatTrackedForSessionInstance( const uint32 StatID, const FGuid& SessionInstanceID ) const
-{
-	bool bResult = false;
-	const FTrackedStat* TrackedStat = TrackedStats.Find( StatID );
-	if( TrackedStat )
-	{
-		bResult = (*TrackedStat).CombinedGraphDataSource->IsProfilerSessionRegistered( SessionInstanceID );
-	}
-
-	return bResult;
-}
-
-
 // @TODO: Move to profiler settings
 const FLinearColor& FProfilerManager::GetColorForStatID( const uint32 StatID ) const
 {
diff --git a/Engine/Source/Developer/Profiler/Private/ProfilerManager.h b/Engine/Source/Developer/Profiler/Private/ProfilerManager.h
--- a/Engine/Source/Developer/Profiler/Private/ProfilerManager.h
+++ b/Engine/Source/Developer/Profiler/Private/ProfilerManager.h
@@ -81,30 +81,25 @@ public:
 
 
 /** Contains basic information about tracked stat. */
-class FTrackedStat //: public FNoncopyable, public TSharedFromThis<FTrackedStat>
+class FTrackedStat : public FNoncopyable, public TSharedFromThis<FTrackedStat>
 {
 public:
 	/**
 	 * Initialization constructor.
 	 *
-	 * @param InColorAverage	- color used to draw the average value
-	 * @param InColorExtremes	- color used to draw the extreme values
-	 * @param InColorBackground - color used to draw the background
-	 * @param InStatID			- the ID of the stat which will be tracked
+	 * @param InGraphDataSource - Data source
+	 * @param InGraphColor - color used to draw the average value
+	 * @param InStatID - the ID of the stat which will be tracked
 	 *
 	 */
 	FTrackedStat
 	(
-		FCombinedGraphDataSourceRef InCombinedGraphDataSource,
-		const FLinearColor InColorAverage,
-		const FLinearColor InColorExtremes,
-		const FLinearColor InColorBackground,
+		FGraphDataSourceRefConst InGraphDataSource,
+		const FLinearColor InGraphColor,
 		const uint32 InStatID	
 	)
-		: CombinedGraphDataSource( InCombinedGraphDataSource )
-		, ColorAverage( InColorAverage )
-		, ColorExtremes( InColorExtremes )
-		, ColorBackground( InColorBackground )
+		: GraphDataSource( InGraphDataSource )
+		, GraphColor( InGraphColor )
 		, StatID( InStatID )
 	{}
 
@@ -113,17 +108,11 @@ public:
 	{}
 
 //protected:
-	/** A shared reference to the combined graph data source for all active profiler session instances for the specified stat ID. */
-	FCombinedGraphDataSourceRef CombinedGraphDataSource;
+	/** A shared reference to the graph data source for active profiler session for the specified stat ID. */
+	FGraphDataSourceRefConst GraphDataSource;
 
-	/** A color to visualize average value for the combined data graph. */
-	const FLinearColor ColorAverage;
-
-	/** A color to visualize extremes values for the combined data graph, min and max. */
-	const FLinearColor ColorExtremes;
-
-	/** A color to visualize background area for the combined data graph. */
-	const FLinearColor ColorBackground;
+	/** A color to visualize graph value for the data graph. */
+	const FLinearColor GraphColor;
 
 	/** The ID of the stat. */
 	const uint32 StatID;
@@ -260,14 +249,6 @@ public:
 		return ProfilerSession;
 	}
 
-	/**
-	 * Creates a combined graph data source which will provide data for graph drawing.
-	 *
-	 * @param StatID - the ID of the stat that will used for generating the combined graph data source.
-	 *
-	 */
-	FCombinedGraphDataSourceRef CreateCombinedGraphDataSource( const uint32 StatID );
-
 	/*-----------------------------------------------------------------------------
 		Stat tracking, Session instance management
 	-----------------------------------------------------------------------------*/
@@ -281,9 +262,6 @@ public:
 	 * @return true, if the specified stat is currently tracked by the profiler.
 	 */
 	const bool IsStatTracked( const uint32 StatID ) const;
-	bool TrackStatForSessionInstance( const uint32 StatID, const FGuid& SessionInstanceID );
-	bool UntrackStatForSessionInstance( const uint32 StatID, const FGuid& SessionInstanceID );
-	const bool IsStatTrackedForSessionInstance( const uint32 StatID, const FGuid& SessionInstanceID ) const;
 
 	/**
 	 * @return True, if the profiler has at least one fully processed capture file
@@ -374,10 +352,10 @@ public:
 	/**
 	 * The event to execute when the status of specified tracked stat has changed.
 	 *
-	 * @param const FTrackedStat&	- a reference to the tracked stat whose status has changed, this reference is valid only within the scope of function
-	 * @param bool bIsTracked		- true, if stat has been added for tracking, false, if stat has been removed from tracking
+	 * @param const FTrackedStatPtr& - a reference to the tracked stat whose status has changed, this reference is valid only within the scope of function
+	 * @param bool bIsTracked - true, if stat has been added for tracking, false, if stat has been removed from tracking
 	 */
-	DECLARE_EVENT_TwoParams( FProfilerManager, FTrackedStatChangedEvent, const FTrackedStat&, bool );
+	DECLARE_EVENT_TwoParams( FProfilerManager, FTrackedStatChangedEvent, const FTrackedStatPtr&, bool );
 	FTrackedStatChangedEvent& OnTrackedStatChanged()
 	{
 		return TrackedStatChangedEvent;
@@ -468,7 +446,6 @@ protected:
 	void ProfilerClient_OnProfilerFileTransfer( const FString& Filename, int64 FileProgress, int64 FileSize );
 
 	void SessionManager_OnInstanceSelectionChanged( const TSharedPtr<ISessionInstanceInfo>& Instance, bool Selected );
-	void SessionManager_OnSelectedSessionChanged( const ISessionInfoPtr& Session );
 
 public:
 	const FLinearColor& GetColorForStatID( const uint32 StatID ) const;
@@ -536,7 +513,7 @@ protected:
 	FFrameAddedEvent FrameAddedEvent;
 
 	/** Contains all currently tracked stats, stored as StatID -> FTrackedStat. */
-	TMap<uint32, FTrackedStat> TrackedStats;
+	TMap<uint32, FTrackedStatPtr> TrackedStats;
 
 	
 	/*-----------------------------------------------------------------------------
diff --git a/Engine/Source/Developer/Profiler/Private/ProfilerSample.h b/Engine/Source/Developer/Profiler/Private/ProfilerSample.h
--- a/Engine/Source/Developer/Profiler/Private/ProfilerSample.h
+++ b/Engine/Source/Developer/Profiler/Private/ProfilerSample.h
@@ -68,25 +68,6 @@ struct EProfilerSampleTypes
 	static FString ToDescription( EProfilerSampleTypes::Type ProfilerSampleType );
 };
 
-/** Enumerates types where profiler samples can be stored. */
-namespace EStorageType
-{
-	enum Type
-	{
-		/** Stored in a local memory. */
-		LocalMemory,
-
-		/** Stored in a file. */
-		BackedUpByFile,
-
-		/** Stored in a global cache. */
-		GlobalCache,
-		
-		/** Invalid enum type, may be used as a number of enumerations. */
-		InvalidOrMax,
-	};
-}
-
 /*-----------------------------------------------------------------------------
 	Declarations
 -----------------------------------------------------------------------------*/
diff --git a/Engine/Source/Developer/Profiler/Private/ProfilerSession.cpp b/Engine/Source/Developer/Profiler/Private/ProfilerSession.cpp
--- a/Engine/Source/Developer/Profiler/Private/ProfilerSession.cpp
+++ b/Engine/Source/Developer/Profiler/Private/ProfilerSession.cpp
@@ -27,8 +27,6 @@ FProfilerSession::FProfilerSession( EProfilerSessionTypes::Type InSessionType, c
 , OnTick( FTickerDelegate::CreateRaw( this, &FProfilerSession::HandleTicker ) )
 , DataProvider( MakeShareable( new FArrayDataProvider() ) )
 , StatMetaData( MakeShareable( new FProfilerStatMetaData() ) )
-, EventGraphDataTotal( MakeShareable( new FEventGraphData() ) )
-, EventGraphDataMaximum( MakeShareable( new FEventGraphData() ) )
 , EventGraphDataCurrent( nullptr )
 , CreationTime( FDateTime::Now() )
 , SessionType( InSessionType )
@@ -51,8 +49,6 @@ FProfilerSession::FProfilerSession( const ISessionInstanceInfoPtr InSessionInsta
 , OnTick( FTickerDelegate::CreateRaw( this, &FProfilerSession::HandleTicker ) )
 , DataProvider( MakeShareable( new FArrayDataProvider() ) )
 , StatMetaData( MakeShareable( new FProfilerStatMetaData() ) )
-, EventGraphDataTotal( MakeShareable( new FEventGraphData() ) )
-, EventGraphDataMaximum( MakeShareable( new FEventGraphData() ) )
 , EventGraphDataCurrent( nullptr )
 , CreationTime( FDateTime::Now() )
 , SessionType( EProfilerSessionTypes::Live )
@@ -79,8 +75,6 @@ FProfilerSession::FProfilerSession( const FString& InDataFilepath )
 , OnTick( FTickerDelegate::CreateRaw( this, &FProfilerSession::HandleTicker ) )
 , DataProvider( MakeShareable( new FArrayDataProvider() ) )
 , StatMetaData( MakeShareable( new FProfilerStatMetaData() ) )
-, EventGraphDataTotal( MakeShareable( new FEventGraphData() ) )
-, EventGraphDataMaximum( MakeShareable( new FEventGraphData() ) )
 , EventGraphDataCurrent( nullptr )
 , CreationTime( FDateTime::Now() )
 , SessionType( EProfilerSessionTypes::StatsFile )
@@ -105,47 +99,31 @@ FProfilerSession::~FProfilerSession()
 	FTicker::GetCoreTicker().RemoveTicker( OnTickHandle );	
 }
 
-FGraphDataSourceRefConst FProfilerSession::CreateGraphDataSource( const uint32 InStatID )
+const FEventGraphDataRef FProfilerSession::GetEventGraphDataTotal() const
 {
-	FGraphDataSource* GraphDataSource = new FGraphDataSource( AsShared(), InStatID );
-	return MakeShareable( GraphDataSource );
+	return EventGraphDataTotal.ToSharedRef();
 }
 
-FEventGraphDataRef FProfilerSession::CreateEventGraphData( const uint32 FrameStartIndex, const uint32 FrameEndIndex, const EEventGraphTypes::Type EventGraphType )
+const FEventGraphDataRef FProfilerSession::GetEventGraphDataMaximum() const
 {
-	static FTotalTimeAndCount Current(0.0f, 0);
-	PROFILER_SCOPE_LOG_TIME( TEXT( "FProfilerSession::CreateEventGraphData" ), &Current );
-
-	FEventGraphData* EventGraphData = new FEventGraphData();
+	return EventGraphDataMaximum.ToSharedRef();
+}
 
-	if( EventGraphType == EEventGraphTypes::Average )
-	{
-		for( uint32 FrameIndex = FrameStartIndex; FrameIndex < FrameEndIndex+1; ++FrameIndex )
-		{
-			// Create a temporary event graph data for the specified frame.
-			const FEventGraphData CurrentEventGraphData( AsShared(), FrameIndex );
-			EventGraphData->CombineAndAdd( CurrentEventGraphData );
-		}
-	
-		EventGraphData->Advance( FrameStartIndex, FrameEndIndex+1 );
-		EventGraphData->Divide( (double)EventGraphData->GetNumFrames() );
-	}
-	else if( EventGraphType == EEventGraphTypes::Maximum )
-	{
-		for( uint32 FrameIndex = FrameStartIndex; FrameIndex < FrameEndIndex+1; ++FrameIndex )
-		{
-			// Create a temporary event graph data for the specified frame.
-			const FEventGraphData CurrentEventGraphData( AsShared(), FrameIndex );
-			EventGraphData->CombineAndFindMax( CurrentEventGraphData );
-		}
-		EventGraphData->Advance( FrameStartIndex, FrameEndIndex+1 );
-	}
+const FEventGraphDataRef FProfilerSession::GetEventGraphDataAverage() const
+{
+	return EventGraphDataAverage.ToSharedRef();
+}
 
-	return MakeShareable( EventGraphData );
+FGraphDataSourceRefConst FProfilerSession::CreateGraphDataSource( const uint32 InStatID )
+{
+	FGraphDataSource* GraphDataSource = new FGraphDataSource( AsShared(), InStatID );
+	return MakeShareable( GraphDataSource );
 }
 
 void FProfilerSession::UpdateAggregatedStats( const uint32 FrameIndex )
 {
+	// #YRX_Profiler 2015-12-08  Move to profiler client manager or ignore hier stats, and move hier stats to PopulateHierarchy_Recurrent
+
 	static FTotalTimeAndCount TimeAndCount( 0.0f, 0 );
 	PROFILER_SCOPE_LOG_TIME( TEXT( "2 FProfilerSession::UpdateAggregatedStats" ), &TimeAndCount );
 
@@ -178,26 +156,65 @@ void FProfilerSession::UpdateAggregatedStats( const uint32 FrameIndex )
 	// @TODO: Create map for stats TMap<uint32, TArray<uint32> >; StatID -> Sample indices for faster lookup in data providers
 }
 
-void FProfilerSession::EventGraphCombineAndMax( const FEventGraphData* Current, const uint32 InNumFrames )
+FEventGraphDataRef FProfilerSession::CreateEventGraphData( const uint32 FrameStartIndex, const uint32 FrameEndIndex, const EEventGraphTypes::Type EventGraphType )
 {
-	EventGraphDataMaximum->CombineAndFindMax( *Current );
-	
-	if (InNumFrames > 0)
+	static FTotalTimeAndCount Current( 0.0f, 0 );
+	PROFILER_SCOPE_LOG_TIME( TEXT( "FProfilerSession::CreateEventGraphData" ), &Current );
+
+	FEventGraphData* EventGraphData = new FEventGraphData( AsShared(), FrameStartIndex );
+	for (uint32 FrameIndex = FrameStartIndex + 1; FrameIndex < FrameEndIndex + 1; ++FrameIndex)
+	{
+		// Create a temporary event graph data for the specified frame.
+		const FEventGraphData CurrentEventGraphData( AsShared(), FrameIndex );
+		EventGraphData->Combine( CurrentEventGraphData );
+	}
+	EventGraphData->Finalize( FrameStartIndex, FrameEndIndex + 1 );
+
+	if (EventGraphType == EEventGraphTypes::Average)
+	{
+		EventGraphData->SetAsAverage();
+	}
+	else if (EventGraphType == EEventGraphTypes::Maximum)
+	{
+
+		EventGraphData->SetAsMaximim();
+	}
+	else if (EventGraphType == EEventGraphTypes::Total)
 	{
-		EventGraphDataMaximum->Advance( 0, InNumFrames );
+		// Nothing to do.
 	}
+
+	return MakeShareable( EventGraphData );
 }
 
-void FProfilerSession::EventGraphCombineAndAdd( const FEventGraphData* Current, const uint32 InNumFrames )
+void FProfilerSession::EventGraphCombine( const FEventGraphData* Current, const uint32 InNumFrames )
 {
-	EventGraphDataTotal->CombineAndAdd( *Current );
-
+	if (EventGraphDataTotal.IsValid())
+	{
+		EventGraphDataTotal->Combine( *Current );
+	}
+	else
+	{
+		EventGraphDataTotal = MakeShareable( new FEventGraphData( *Current ) );
+	}
+	
 	if (InNumFrames > 0)
 	{
-		EventGraphDataTotal->Advance( 0, InNumFrames );
+		UpdateAllEventGraphs( InNumFrames );
 	}
 }
 
+void FProfilerSession::UpdateAllEventGraphs( const uint32 InNumFrames )
+{
+	EventGraphDataTotal->Finalize( 0, DataProvider->GetNumFrames() );
+	
+	EventGraphDataAverage = EventGraphDataTotal->DuplicateAsRef();
+	EventGraphDataAverage->SetAsAverage();
+
+	EventGraphDataMaximum = EventGraphDataTotal->DuplicateAsRef();
+	EventGraphDataMaximum->SetAsMaximim();
+}
+
 void FProfilerSession::UpdateAggregatedEventGraphData( const uint32 FrameIndex )
 {
 	static FTotalTimeAndCount TimeAndCount( 0.0f, 0 );
@@ -210,43 +227,24 @@ void FProfilerSession::UpdateAggregatedEventGraphData( const uint32 FrameIndex )
 	EventGraphDataCurrent = new FEventGraphData( AsShared(), FrameIndex );
 	const uint32 NumFramesLocal = SessionType == EProfilerSessionTypes::Live ? DataProvider->GetNumFrames() : 0;
 
-	static const bool bUseTaskGraph = true;
+	static const bool bUseTaskGraph = false;
 
 	if( bUseTaskGraph )
 	{
-		FGraphEventArray EventGraphCombineTasks;
-
-		DECLARE_CYCLE_STAT(TEXT("FSimpleDelegateGraphTask.EventGraphData.CombineAndFindMax"),
-			STAT_FSimpleDelegateGraphTask_EventGraphData_CombineAndFindMax,
-			STATGROUP_TaskGraphTasks);
-		DECLARE_CYCLE_STAT(TEXT("FSimpleDelegateGraphTask.EventGraphData.CombineAndAdd"),
-			STAT_FSimpleDelegateGraphTask_EventGraphData_EventGraphCombineAndAdd,
+		DECLARE_CYCLE_STAT(TEXT("FSimpleDelegateGraphTask.EventGraphData.GraphCombine"),
+			STAT_FSimpleDelegateGraphTask_EventGraphData_GraphCombine,
 			STATGROUP_TaskGraphTasks);
 
-		new (EventGraphCombineTasks) FGraphEventRef(FSimpleDelegateGraphTask::CreateAndDispatchWhenReady
-		(
-			FSimpleDelegateGraphTask::FDelegate::CreateRaw( this, &FProfilerSession::EventGraphCombineAndMax, EventGraphDataCurrent, NumFramesLocal ), 
-			GET_STATID(STAT_FSimpleDelegateGraphTask_EventGraphData_CombineAndFindMax), nullptr
-		));
-
-		new (EventGraphCombineTasks) FGraphEventRef(FSimpleDelegateGraphTask::CreateAndDispatchWhenReady
+		CompletionSync = FSimpleDelegateGraphTask::CreateAndDispatchWhenReady
 		(
-			FSimpleDelegateGraphTask::FDelegate::CreateRaw( this, &FProfilerSession::EventGraphCombineAndAdd, EventGraphDataCurrent, NumFramesLocal ), 
-			GET_STATID(STAT_FSimpleDelegateGraphTask_EventGraphData_EventGraphCombineAndAdd), nullptr
-		));
-
-		DECLARE_CYCLE_STAT(TEXT("FNullGraphTask.EventGraphData.CombineJoinAndContinue"),
-			STAT_FNullGraphTask_EventGraphData_CombineJoinAndContinue,
-			STATGROUP_TaskGraphTasks);
-
-		// JoinThreads
-		CompletionSync = TGraphTask<FNullGraphTask>::CreateTask( &EventGraphCombineTasks, ENamedThreads::GameThread )
-			.ConstructAndDispatchWhenReady(GET_STATID(STAT_FNullGraphTask_EventGraphData_CombineJoinAndContinue), ENamedThreads::AnyThread);
+			FSimpleDelegateGraphTask::FDelegate::CreateRaw( this, &FProfilerSession::EventGraphCombine, EventGraphDataCurrent, NumFramesLocal ), 
+			GET_STATID( STAT_FSimpleDelegateGraphTask_EventGraphData_GraphCombine )
+		);
+		
 	}
 	else
 	{
-		EventGraphCombineAndMax( EventGraphDataCurrent, NumFramesLocal );
-		EventGraphCombineAndAdd( EventGraphDataCurrent, NumFramesLocal );
+		EventGraphCombine( EventGraphDataCurrent, NumFramesLocal );
 	}
 }
 
@@ -261,6 +259,8 @@ void FProfilerSession::CompletionSyncAggregatedEventGraphData()
 	}
 }
 
+
+
 bool FProfilerSession::HandleTicker( float DeltaTime )
 {
 	// Update metadata if needed
@@ -426,8 +426,7 @@ bool FProfilerSession::HandleTicker( float DeltaTime )
 			CompletionSyncAggregatedEventGraphData();
 
 			// Advance event graphs.
-			EventGraphDataMaximum->Advance( 0, DataProvider->GetNumFrames() );
-			EventGraphDataTotal->Advance( 0, DataProvider->GetNumFrames() );
+			UpdateAllEventGraphs( DataProvider->GetNumFrames() );
 
 			// Broadcast that a capture file has been fully processed.
 			OnCaptureFileProcessed.ExecuteIfBound( GetInstanceID() );
@@ -687,7 +686,6 @@ void FRawProfilerSession::PrepareLoading()
 	const bool bIsFinalized = Stream.Header.IsFinalized();
 	check( bIsFinalized );
 	check( Stream.Header.Version == EStatMagicWithHeader::VERSION_5 );
-	StatsThreadStats.MarkAsLoaded();
 
 	TArray<FStatMessage> Messages;
 	if( Stream.Header.bRawStatsFile )
diff --git a/Engine/Source/Developer/Profiler/Private/ProfilerSession.h b/Engine/Source/Developer/Profiler/Private/ProfilerSession.h
--- a/Engine/Source/Developer/Profiler/Private/ProfilerSession.h
+++ b/Engine/Source/Developer/Profiler/Private/ProfilerSession.h
@@ -1064,20 +1064,9 @@ public:
 		return AggregatedStats.Find( StatID );
 	}
 
-	const FEventGraphDataRef& GetEventGraphDataMaximum() const
-	{
-		return EventGraphDataMaximum;
-	}
-
-	const FEventGraphDataRef GetEventGraphDataAverage() const
-	{
-		const uint32 NumFramesLocal = EventGraphDataTotal->GetNumFrames();
-		FEventGraphDataRef EventGraphDataAverage = EventGraphDataTotal->DuplicateAsRef();
-		EventGraphDataAverage->Divide( (double)NumFramesLocal );
-		EventGraphDataAverage->Advance( 0, NumFramesLocal );
-
-		return EventGraphDataAverage;
-	}
+	const FEventGraphDataRef GetEventGraphDataTotal() const;
+	const FEventGraphDataRef GetEventGraphDataMaximum() const;
+	const FEventGraphDataRef GetEventGraphDataAverage() const;
 
 	/**
 	 * @return number of bytes allocated by this profiler session.
@@ -1145,9 +1134,9 @@ protected:
 	/** Completion sync. */
 	void CompletionSyncAggregatedEventGraphData();
 
-	void EventGraphCombineAndMax( const FEventGraphData* Current, const uint32 NumFrames );
+	void EventGraphCombine( const FEventGraphData* Current, const uint32 InNumFrames );
 
-	void EventGraphCombineAndAdd( const FEventGraphData* Current, const uint32 NumFrames );
+	void UpdateAllEventGraphs( const uint32 InNumFrames );
 
 	/** Called when the capture file has been fully loaded. */
 	void LoadComplete();
@@ -1195,13 +1184,14 @@ protected:
 	/** The stat metadata which holds all collected stats descriptions. */
 	FProfilerStatMetaDataRef StatMetaData;
 
-	// TODO: Make this more universal ? FEventGraphDataRef<AggregateFunc,FinalizerFunc,bRequiredCopy>
+	/** Aggregated event graph data for all collected frames, used for generating average values. Also contains min and max. */
+	FEventGraphDataPtr EventGraphDataTotal;
 
-	/** Aggregated event graph data for all collected frames, used for generating average values. */
-	FEventGraphDataRef EventGraphDataTotal;
+	/** Highest "per-frame" event graph. */
+	FEventGraphDataPtr EventGraphDataMaximum;
 
-	/** Aggregated event graph data for all collected frames, contains only maximum inclusive times. */
-	FEventGraphDataRef EventGraphDataMaximum;
+	/** Per-frame average event graph. */
+	FEventGraphDataPtr EventGraphDataAverage;
 
 	/** Temporary event graph data for the specified frame. Recreated each frame, used by the task graph tasks only. */
 	const FEventGraphData* EventGraphDataCurrent;
@@ -1260,7 +1250,7 @@ class FRawProfilerSession : public FProfilerSession
 	FProfilerStream ProfilerStream;
 
 	/** Stats thread state, mostly used to manage the stats metadata. */
-	FStatsThreadState StatsThreadStats;
+	FStatsLoadedState StatsThreadStats;
 	FStatsReadStream Stream;
 
 	/** Index of the last processed data for the mini-view. */
diff --git a/Engine/Source/Developer/Profiler/Private/Widgets/SDataGraph.cpp b/Engine/Source/Developer/Profiler/Private/Widgets/SDataGraph.cpp
--- a/Engine/Source/Developer/Profiler/Private/Widgets/SDataGraph.cpp
+++ b/Engine/Source/Developer/Profiler/Private/Widgets/SDataGraph.cpp
@@ -22,11 +22,11 @@ public:
 
 	SLATE_BEGIN_ARGS( SDataGraphSummary )
 		: _ParentWidget()
-		, _GraphDescription()
+		, _TrackedStat()
 		{}
 
 		SLATE_ARGUMENT( TSharedPtr<SDataGraph>, ParentWidget )
-		SLATE_ARGUMENT( FGraphDescription, GraphDescription )
+		SLATE_ARGUMENT( FTrackedStatPtr, TrackedStat )
 		SLATE_EVENT( FGetHoveredFrameIndexDelegate, OnGetMouseFrameIndex )
 	SLATE_END_ARGS()
 
@@ -39,13 +39,13 @@ public:
 	void Construct( const FArguments& InArgs )
 	{
 		ParentWidget = InArgs._ParentWidget;
-		GraphDescription = InArgs._GraphDescription;
+		TrackedStat = InArgs._TrackedStat;
 		OnGetMouseFrameIndex = InArgs._OnGetMouseFrameIndex;
 
-		const FSlateColor TextColor( GraphDescription.ColorAverage );
+		const FSlateColor TextColor( TrackedStat->GraphColor );
 		const FName CloseButtonStyle = TEXT("Docking.MajorTab.CloseButton");
 		FFormatNamedArguments Args;
-		Args.Add( TEXT("StatName"), FText::FromString( GraphDescription.CombinedGraphDataSource->GetStatName() ) );
+		Args.Add( TEXT("StatName"), FText::FromString( TrackedStat->GraphDataSource->GetStatName() ) );
 		const FText ToolTipText = FText::Format( LOCTEXT("DataGraphSummary_CloseButton_TT", "Click to stop tracking '{StatName}' stat"), Args );
 
 		ChildSlot
@@ -116,23 +116,23 @@ protected:
 	/** Stops tracking the associated stat and removes it from the data graph. */
 	FReply CloseButton_OnClicked()
 	{
-		FProfilerManager::Get()->UntrackStat( GraphDescription.CombinedGraphDataSource->GetStatID() );	
+		FProfilerManager::Get()->UntrackStat( TrackedStat->GraphDataSource->GetStatID() );	
 		return FReply::Handled();
 	}
 
 	FText SummaryInformation_GetSummary() const
 	{
 		FText SummaryText = LOCTEXT("DataGraphSummary_Warning", "Not implemented yet");
-		const bool bCanDisplayData = GraphDescription.CombinedGraphDataSource->CanBeDisplayedAsIndexBased() && ParentWidget->GetViewMode() == EDataGraphViewModes::Index;
+		const bool bCanDisplayData = TrackedStat->GraphDataSource->CanBeDisplayedAsIndexBased() && ParentWidget.Pin()->GetViewMode() == EDataGraphViewModes::Index;
 
 		const uint32 FrameIndex = OnGetMouseFrameIndex.IsBound() ? (uint32)OnGetMouseFrameIndex.Execute() : 0;
-		const FGraphDataSourceRefConst* GraphDataSource = GraphDescription.CombinedGraphDataSource->GetFirstSource();
+		const FGraphDataSourceRefConst& GraphDataSource = TrackedStat->GraphDataSource;
 
-		if( bCanDisplayData && GraphDataSource && FrameIndex < (*GraphDataSource)->GetNumFrames() )
+		if( bCanDisplayData && FrameIndex < GraphDataSource->GetNumFrames() )
 		{
-			const double SampleValue = (*GraphDataSource)->GetValueFromIndex( FrameIndex );
-			const EProfilerSampleTypes::Type UnitType = (*GraphDataSource)->GetSampleType();
-			const FProfilerAggregatedStat& Aggregated = *(*GraphDataSource)->GetAggregatedStat();
+			const double SampleValue = GraphDataSource->GetValueFromIndex( FrameIndex );
+			const EProfilerSampleTypes::Type UnitType = GraphDataSource->GetSampleType();
+			const FProfilerAggregatedStat& Aggregated = *GraphDataSource->GetAggregatedStat();
 
 			static const FNumberFormattingOptions SampleValueFormattingOptions = FNumberFormattingOptions()
 				.SetMinimumFractionalDigits(2)
@@ -146,19 +146,20 @@ protected:
 
 	FText SummaryInformation_GetGroupName() const
 	{
-		return FText::Format(LOCTEXT("DataGraphSummary_GroupNameFmt", "({0})"), FText::FromString(GraphDescription.CombinedGraphDataSource->GetGroupName()));
+		return FText::Format(LOCTEXT("DataGraphSummary_GroupNameFmt", "({0})"), FText::FromString(TrackedStat->GraphDataSource->GetGroupName()));
 	}
 
 	FText SummaryInformation_GetStatName() const
 	{
-		return FText::FromString(FProfilerHelper::ShortenName( GraphDescription.CombinedGraphDataSource->GetStatName(), 32 ));
+		return FText::FromString(FProfilerHelper::ShortenName( TrackedStat->GraphDataSource->GetStatName(), 32 ));
 	}
 
 private:
-	FGraphDescription GraphDescription;
+	/** A shared pointer to the traced stat. */
+	FTrackedStatPtr TrackedStat;
 
-	/** A shared pointer to the parent widget. */
-	TSharedPtr<SDataGraph> ParentWidget;
+	/** A weak pointer to the parent widget. */
+	TWeakPtr<SDataGraph> ParentWidget;
 
 	/** The delegate to be invoked when the data graph summary widget wants to know the frame index pointed by the mouse. */
 	FGetHoveredFrameIndexDelegate OnGetMouseFrameIndex;
@@ -236,32 +237,24 @@ void SDataGraph::Tick( const FGeometry& AllottedGeometry, const double InCurrent
 
 void SDataGraph::UpdateState()
 {
-	const FGraphDescription* GraphDesc = GetFirstGraph();
-	if( GraphDesc )
+	const FTrackedStatPtr FirstTrackedStat = GetFirstGraph();
+	if( FirstTrackedStat.IsValid() )
 	{
-		// Check if we need to force time based view mode.
-		const bool bCanBeDisplayedAsMulti = false;
-		if( bCanBeDisplayedAsMulti )
-		{
-			ViewMode = EDataGraphViewModes::Time;
-		}
-
 		// If the view mode is index based, use the first source for reading the number of frames.
 		if( ViewMode == EDataGraphViewModes::Index )
 		{
-			const FGraphDataSourceRefConst* GraphDataSource = GraphDesc->CombinedGraphDataSource->GetFirstSource();
-			NumDataPoints = GraphDataSource ? (int32)(*GraphDataSource)->GetNumFrames() : 0;
+			NumDataPoints = FirstTrackedStat->GraphDataSource->GetNumFrames();
 		}
 		else
 		{
-			NumDataPoints = (int32)GraphDesc->CombinedGraphDataSource->GetNumFrames();
+			NumDataPoints = (int32)FirstTrackedStat->GraphDataSource->GetNumFrames();
 		}
 
 		NumVisiblePoints = FMath::Max( 0, FMath::TruncToInt(ThisGeometry.Size.X) / DistanceBetweenPoints );
 		// GraphOffset - Updated by OnMouseMove or by ScrollTo
 		GraphOffset = FMath::Clamp( GraphOffset, 0, FMath::Max(NumDataPoints-NumVisiblePoints,0) );
 		
-		DataTotalTimeMS = GraphDesc->CombinedGraphDataSource->GetTotalTimeMS();
+		DataTotalTimeMS = FirstTrackedStat->GraphDataSource->GetTotalTimeMS();
 		VisibleTimeMS = NumVisiblePoints * FTimeAccuracy::AsFrameTime( TimeBasedAccuracy );
 		GraphOffsetMS = GraphOffset * FTimeAccuracy::AsFrameTime( TimeBasedAccuracy );
 	}
@@ -324,22 +317,22 @@ int32 SDataGraph::OnPaint( const FPaintArgs& Args, const FGeometry& AllottedGeom
 	{
 		SCOPE_CYCLE_COUNTER(STAT_DG_OnPaint);
 
-		const FGraphDescription& GraphDescription = It.Value();
+		const FTrackedStat& TrackedStat = *It.Value();
 		const float GraphYScale = AllottedGeometry.Size.Y/ScaleY;
 		
-		const float UnitTypeScale = GraphDescription.CombinedGraphDataSource->GetSampleType() != EProfilerSampleTypes::Memory ? 1.0f : CounterToTimeScale;
+		const float UnitTypeScale = TrackedStat.GraphDataSource->GetSampleType() != EProfilerSampleTypes::Memory ? 1.0f : CounterToTimeScale;
 		const float TimeAccuracyMS = FTimeAccuracy::AsFrameTime( TimeBasedAccuracy );
 
 		if( ViewMode == EDataGraphViewModes::Time )
 		{
 			const float GraphRangeEndMS = FMath::Min( GraphOffsetMS+VisibleTimeMS, DataTotalTimeMS ) - TimeAccuracyMS;
 
-			if( MultiMode == EDataGraphMultiModes::Combined && GraphDescription.CombinedGraphDataSource->GetSourcesNum() > 0 )
+			/*if( MultiMode == EDataGraphMultiModes::Combined && TrackedStat.CombinedGraphDataSource->GetSourcesNum() > 0 )
 			{
 				// Draw combined line graph where X=Min, Y=Max, Z=Avg
 				for( float GraphStartTimeMS = GraphOffsetMS; GraphStartTimeMS < GraphRangeEndMS; GraphStartTimeMS += TimeAccuracyMS )
 				{
-					const FVector Value = GraphDescription.CombinedGraphDataSource->GetValueFromTimeRange( GraphStartTimeMS, GraphStartTimeMS+TimeAccuracyMS );
+					const FVector Value = TrackedStat.CombinedGraphDataSource->GetValueFromTimeRange( GraphStartTimeMS, GraphStartTimeMS+TimeAccuracyMS );
 					const float XPos = DistanceBetweenPoints*GraphPoints.Num();
 
 					// X=Min
@@ -370,7 +363,7 @@ int32 SDataGraph::OnPaint( const FPaintArgs& Args, const FGeometry& AllottedGeom
 					GraphPoints,
 					MyClippingRect,
 					DrawEffects,
-					InWidgetStyle.GetColorAndOpacityTint() * GraphDescription.ColorBackground,
+					InWidgetStyle.GetColorAndOpacityTint() * TrackedStat.ColorBackground,
 					false
 				);
 				GraphPoints.Empty( NumVisiblePoints );
@@ -384,7 +377,7 @@ int32 SDataGraph::OnPaint( const FPaintArgs& Args, const FGeometry& AllottedGeom
 					GraphPoints2,
 					MyClippingRect,
 					DrawEffects,
-					InWidgetStyle.GetColorAndOpacityTint() * GraphDescription.ColorExtremes,
+					InWidgetStyle.GetColorAndOpacityTint() * TrackedStat.ColorExtremes,
 					false
 				);
 				GraphPoints2.Empty( NumVisiblePoints );
@@ -398,75 +391,72 @@ int32 SDataGraph::OnPaint( const FPaintArgs& Args, const FGeometry& AllottedGeom
 					GraphPoints3,
 					MyClippingRect,
 					DrawEffects,
-					InWidgetStyle.GetColorAndOpacityTint() * GraphDescription.ColorAverage,
+					InWidgetStyle.GetColorAndOpacityTint() * TrackedStat.ColorAverage,
 					false
 				);
 				GraphPoints3.Empty( NumVisiblePoints );
 				
 				LayerId++;
 			}
-			else if( MultiMode == EDataGraphMultiModes::OneLinePerDataSource )
+			else*/
+			if( MultiMode == EDataGraphMultiModes::OneLinePerDataSource )
 			{
 				// Draw line graph for each graph data source.
-				for( auto GraphSourceIt = GraphDescription.CombinedGraphDataSource->GetSourcesIterator(); GraphSourceIt; ++GraphSourceIt )
+
+				const FGraphDataSourceRefConst& GraphDataSource = TrackedStat.GraphDataSource;
+
+				for (float GraphStartTimeMS = GraphOffsetMS; GraphStartTimeMS < GraphRangeEndMS; GraphStartTimeMS += TimeAccuracyMS)
 				{
-					const FGraphDataSourceRefConst& GraphDataSource = GraphSourceIt.Value();
- 
-					for( float GraphStartTimeMS = GraphOffsetMS; GraphStartTimeMS < GraphRangeEndMS; GraphStartTimeMS += TimeAccuracyMS )
-					{
-						const float Value = GraphDataSource->GetValueFromTimeRange( GraphStartTimeMS, GraphStartTimeMS+TimeAccuracyMS );
-						const float XPos = DistanceBetweenPoints*GraphPoints.Num();
-						const float YPos = FMath::Clamp( AllottedGeometry.Size.Y - GraphYScale*Value*UnitTypeScale, 0.0f, AllottedGeometry.Size.Y );
-						new (GraphPoints) FVector2D(XPos,YPos);
-					}
- 
-					FSlateDrawElement::MakeLines
-					(
-						OutDrawElements,
-						LayerId,
-						AllottedGeometry.ToPaintGeometry(),
-						GraphPoints,
-						MyClippingRect,
-						DrawEffects,
-						InWidgetStyle.GetColorAndOpacityTint() * GraphDescription.ColorAverage,
-						false
-					);
-					GraphPoints.Empty( NumVisiblePoints );
+					const float Value = GraphDataSource->GetValueFromTimeRange( GraphStartTimeMS, GraphStartTimeMS + TimeAccuracyMS );
+					const float XPos = DistanceBetweenPoints*GraphPoints.Num();
+					const float YPos = FMath::Clamp( AllottedGeometry.Size.Y - GraphYScale*Value*UnitTypeScale, 0.0f, AllottedGeometry.Size.Y );
+					new (GraphPoints)FVector2D( XPos, YPos );
 				}
-				LayerId++;
+
+				FSlateDrawElement::MakeLines
+				(
+					OutDrawElements,
+					LayerId,
+					AllottedGeometry.ToPaintGeometry(),
+					GraphPoints,
+					MyClippingRect,
+					DrawEffects,
+					InWidgetStyle.GetColorAndOpacityTint() * TrackedStat.GraphColor,
+					false
+				);
+				GraphPoints.Empty( NumVisiblePoints );
 			}
+			LayerId++;
+
 		}
 		else if( ViewMode == EDataGraphViewModes::Index )
 		{
 			if( MultiMode == EDataGraphMultiModes::OneLinePerDataSource )
-			{
-				for( auto GraphSourceIt = GraphDescription.CombinedGraphDataSource->GetSourcesIterator(); GraphSourceIt; ++GraphSourceIt )
-				{			
-					const FGraphDataSourceRefConst& GraphDataSource = GraphSourceIt.Value();
-					const int32 GraphRangeEndIndex = FMath::Min( GraphOffset+NumVisiblePoints+1, NumDataPoints );
- 
-					for( uint32 GraphStartIndex = (uint32)GraphOffset; GraphStartIndex < (uint32)GraphRangeEndIndex; GraphStartIndex++ )
-					{
-						const float Value = GraphDataSource->GetValueFromIndex( GraphStartIndex );
-						const float XPos = DistanceBetweenPoints*(float)GraphPoints.Num();
-						const float YPos = FMath::Clamp( AllottedGeometry.Size.Y - GraphYScale*Value*UnitTypeScale, 0.0f, AllottedGeometry.Size.Y );
-						GraphPoints.Add( FVector2D(XPos,YPos) );
-					}
+			{						
+				const FGraphDataSourceRefConst& GraphDataSource = TrackedStat.GraphDataSource;
+				const int32 GraphRangeEndIndex = FMath::Min( GraphOffset+NumVisiblePoints+1, NumDataPoints );
  
-					FSlateDrawElement::MakeLines
-					(
-						OutDrawElements,
-						LayerId,
-						AllottedGeometry.ToPaintGeometry(),
-						GraphPoints,
-						MyClippingRect,
-						DrawEffects,
-						InWidgetStyle.GetColorAndOpacityTint() * GraphDescription.ColorAverage,
-						false
-					);
-					GraphPoints.Empty( NumVisiblePoints );
+				for( uint32 GraphStartIndex = (uint32)GraphOffset; GraphStartIndex < (uint32)GraphRangeEndIndex; GraphStartIndex++ )
+				{
+					const float Value = GraphDataSource->GetValueFromIndex( GraphStartIndex );
+					const float XPos = DistanceBetweenPoints*(float)GraphPoints.Num();
+					const float YPos = FMath::Clamp( AllottedGeometry.Size.Y - GraphYScale*Value*UnitTypeScale, 0.0f, AllottedGeometry.Size.Y );
+					GraphPoints.Add( FVector2D(XPos,YPos) );
 				}
  
+				FSlateDrawElement::MakeLines
+				(
+					OutDrawElements,
+					LayerId,
+					AllottedGeometry.ToPaintGeometry(),
+					GraphPoints,
+					MyClippingRect,
+					DrawEffects,
+					InWidgetStyle.GetColorAndOpacityTint() * TrackedStat.GraphColor,
+					false
+				);
+				GraphPoints.Empty( NumVisiblePoints );
+				
 				LayerId++;
 			}
 		}
@@ -492,18 +482,18 @@ int32 SDataGraph::OnPaint( const FPaintArgs& Args, const FGeometry& AllottedGeom
 	TArray<FVector2D> LinePoints;
 	const float LabelSize = MaxFontCharHeight * 7.0f;
 
-	const FGraphDescription* GraphDesc = GetFirstGraph();
-	if( GraphDesc && GraphDesc->CombinedGraphDataSource->GetFirstSource() )
+	const FTrackedStatPtr FirstTrackedStat = GetFirstGraph();
+	if( FirstTrackedStat.IsValid() )
 	{
-		const FGraphDataSourceRefConst* GraphDataSource = GraphDesc->CombinedGraphDataSource->GetFirstSource();
+		const FGraphDataSourceRefConst& GraphDataSource = FirstTrackedStat->GraphDataSource;
+		const IDataProviderRef DataProvider = GraphDataSource->GetDataProvider();
 
 		if( ViewMode == EDataGraphViewModes::Index )
 		{
 			// Draw a vertical line every 60 frames.
 			const int32 AvgFrameRate = 60;
 			const int32 FrameStartIndex = GraphOffset + AvgFrameRate - (GraphOffset % AvgFrameRate);
-			const int32 FrameEndIndex = FMath::Min( GraphOffset + NumVisiblePoints, NumDataPoints );
-			const IDataProviderRef DataProvider = (*GraphDataSource)->GetDataProvider();
+			const int32 FrameEndIndex = FMath::Min( GraphOffset + NumVisiblePoints, NumDataPoints );		
 
 			for( int32 FrameIndex = FrameStartIndex; FrameIndex < FrameEndIndex; FrameIndex += AvgFrameRate )
 			{
@@ -567,7 +557,6 @@ int32 SDataGraph::OnPaint( const FPaintArgs& Args, const FGeometry& AllottedGeom
 			const int32 FrameEndIndex = FMath::Min( GraphOffset + NumVisiblePoints, NumDataPoints );
 
 			const bool bCanBeDisplayedAsMulti = false;
-			const IDataProviderRef DataProvider = (*GraphDataSource)->GetDataProvider();
 
 			for( int32 FrameIndex = FrameStartIndex; FrameIndex < FrameEndIndex; FrameIndex += AvgFrameRate )
 			{
@@ -969,17 +958,8 @@ int32 SDataGraph::OnPaint( const FPaintArgs& Args, const FGeometry& AllottedGeom
 
 
 
-void SDataGraph::AddInnerGraph
-(
-	const uint32 InStatID, 	
-	const FLinearColor InColorAverage,
-	const FLinearColor InColorExtremes,
-	const FLinearColor InColorBackground, 
-	const FCombinedGraphDataSourceRef& CombinedGraphDataSource  
-)
+void SDataGraph::AddInnerGraph( const FTrackedStatPtr& TrackedStat )
 {
-	FGraphDescription GraphDescriptionLine( CombinedGraphDataSource, InColorAverage, InColorExtremes, InColorBackground );
-
 	TSharedPtr<SWidget> GraphSummary;
 	GraphDescriptionsVBox->AddSlot()
 	.AutoHeight()
@@ -989,23 +969,22 @@ void SDataGraph::AddInnerGraph
 	[
 		SAssignNew(GraphSummary,SDataGraphSummary)
 		.ParentWidget( SharedThis(this) )
-		.GraphDescription( GraphDescriptionLine )
+		.TrackedStat( TrackedStat )
 		.OnGetMouseFrameIndex( this, &SDataGraph::DataGraphSummary_GetHoveredFrameIndex )
 	];
-
-	StatIDToGraphDescriptionMapping.Add( InStatID, GraphDescriptionLine );
-	StatIDToWidgetMapping.Add( InStatID, GraphSummary.ToSharedRef() );
+ 
+	StatIDToGraphDescriptionMapping.Add( TrackedStat->StatID, TrackedStat );
+	StatIDToWidgetMapping.Add( TrackedStat->StatID, GraphSummary.ToSharedRef() );
 
 	UpdateState();
 }
 
 void SDataGraph::RemoveInnerGraph( const uint32 InStatID )
 {
-	FGraphDescription* GraphDescription = StatIDToGraphDescriptionMapping.Find( InStatID );
-	if( GraphDescription )
+	FTrackedStatPtr TrackedStat = StatIDToGraphDescriptionMapping.FindRef( InStatID );
+	if( TrackedStat.IsValid() )
 	{
 		TSharedRef<SWidget> DataGraphSummary = StatIDToWidgetMapping.FindChecked( InStatID );
-
 		GraphDescriptionsVBox->RemoveSlot( DataGraphSummary );
 		
 		StatIDToWidgetMapping.Remove( InStatID );
@@ -1179,28 +1158,24 @@ FReply SDataGraph::OnMouseWheel( const FGeometry& MyGeometry, const FPointerEven
 
 const FEventGraphDataHandlerRef SDataGraph::PrepareEventGraphDataHandler( const FVector2D& ScreenSpacePosition )
 {
-	const FGraphDescription* GraphDescription = GetFirstGraph();
-
-	if( GraphDescription )
+	const FTrackedStatPtr FirstTrackedStat = GetFirstGraph();
+	if (FirstTrackedStat.IsValid())
 	{
 		if( ViewMode == EDataGraphViewModes::Time )
 		{
 			TMap<FGuid,uint32> StartIndices;
 			const float TimeAccuracyMS = FTimeAccuracy::AsFrameTime( TimeBasedAccuracy );
-			GraphDescription->CombinedGraphDataSource->GetStartIndicesFromTimeRange( FrameTimesMS[0], FrameTimesMS[1], StartIndices );
+			// #YRX_ProfilerFix? 2015-12-08 
+			//GraphDescription->GraphDataSource->GetStartIndicesFromTimeRange( FrameTimesMS[0], FrameTimesMS[1], StartIndices );
 
 			const FEventGraphDataHandlerRef Params = MakeShareable( new FEventGraphDataHandler( /*StartIndices*/ ) ); 
 			return Params;
 		}
 		else if( ViewMode == EDataGraphViewModes::Index )
 		{
-			const FGraphDataSourceRefConst* GraphDataSource = GraphDescription->CombinedGraphDataSource->GetFirstSource();
-
-			if( GraphDataSource )
-			{
-				const FEventGraphDataHandlerRef Params = MakeShareable( new FEventGraphDataHandler( (*GraphDataSource)->GetSessionInstanceID(), FrameIndices[0], FrameIndices[1], EEventGraphTypes::Maximum ) ); 
-				return Params;
-			}
+			const FGraphDataSourceRefConst& GraphDataSource = FirstTrackedStat->GraphDataSource;
+			const FEventGraphDataHandlerRef Params = MakeShareable( new FEventGraphDataHandler( GraphDataSource->GetSessionInstanceID(), FrameIndices[0], FrameIndices[1], EEventGraphTypes::Maximum ) ); 
+			return Params;
 		}
 	}
 
@@ -1348,7 +1323,7 @@ void SDataGraph::ViewMode_SetIndexBased_Execute()
 
 bool SDataGraph::ViewMode_SetIndexBased_CanExecute() const
 {
-	const bool bCanBeDisplayedAsIndexBased = GetFirstGraph() ? GetFirstGraph()->CombinedGraphDataSource->CanBeDisplayedAsIndexBased() : false;
+	const bool bCanBeDisplayedAsIndexBased = GetFirstGraph().IsValid() ? GetFirstGraph()->GraphDataSource->CanBeDisplayedAsIndexBased() : false;
 	return ViewMode != EDataGraphViewModes::Index && bCanBeDisplayedAsIndexBased;
 }
 
@@ -1370,7 +1345,7 @@ void SDataGraph::ViewMode_SetTimeBased_Execute()
 
 bool SDataGraph::ViewMode_SetTimeBased_CanExecute() const
 {
-	const bool bCanBeDisplayedAsTimeBased = GetFirstGraph() ? GetFirstGraph()->CombinedGraphDataSource->CanBeDisplayedAsTimeBased() : false;
+	const bool bCanBeDisplayedAsTimeBased = GetFirstGraph().IsValid() ? GetFirstGraph()->GraphDataSource->CanBeDisplayedAsTimeBased() : false;
 	return ViewMode != EDataGraphViewModes::Time && bCanBeDisplayedAsTimeBased;
 }
 
@@ -1379,6 +1354,17 @@ bool SDataGraph::ViewMode_SetTimeBased_IsChecked() const
 	return ViewMode == EDataGraphViewModes::Time;
 }
 
+const FTrackedStatPtr SDataGraph::GetFirstGraph() const
+{
+	FTrackedStatPtr Result;
+	for (const auto& It : StatIDToGraphDescriptionMapping)
+	{
+		Result = It.Value; 
+		break;
+	}
+	return Result;
+}
+
 void SDataGraph::EventGraph_OnRestoredFromHistory( uint32 FrameStartIndex, uint32 FrameEndIndex )
 {
 	UpdateState();
diff --git a/Engine/Source/Developer/Profiler/Private/Widgets/SDataGraph.h b/Engine/Source/Developer/Profiler/Private/Widgets/SDataGraph.h
--- a/Engine/Source/Developer/Profiler/Private/Widgets/SDataGraph.h
+++ b/Engine/Source/Developer/Profiler/Private/Widgets/SDataGraph.h
@@ -57,47 +57,6 @@ DECLARE_DELEGATE_OneParam( FViewModeChangedDelegate, EDataGraphViewModes::Type )
 DECLARE_DELEGATE_TwoParams( FSelectionChangedForTimeDelegate, float /*FrameStartTimeMS*/, float /*FrameEndTimeMS*/ );
 
 
-/*-----------------------------------------------------------------------------
-	Basic structures
------------------------------------------------------------------------------*/
-
-// @TODO: GraphDataSource will be replaced later with more abstract class which will allow to use profiler data source or other source to draw the graph.
-
-/** Describes properties of the graph that will be displayed in the SDataGraph widget. */
-class FGraphDescription
-{
-public:
-	FGraphDescription()
-	{}
-
-	FGraphDescription
-	( 
-		const FCombinedGraphDataSourceRef& InGraphDataSource, 
-		const FLinearColor InColorAverage,
-		const FLinearColor InColorExtremes,
-		const FLinearColor InColorBackground
-	)
-		: CombinedGraphDataSource( InGraphDataSource )
-		, ColorAverage( InColorAverage )
-		, ColorExtremes( InColorExtremes )
-		, ColorBackground( InColorBackground )
-	{}
-
-	/** Data source for this graph. */
-	FCombinedGraphDataSourcePtr CombinedGraphDataSource;
-
-	/** A color to visualize average value for the combined data graph. */
-	FLinearColor ColorAverage;
-
-	/** A color to visualize extremes values for the combined data graph, min and max. */
-	FLinearColor ColorExtremes;
-
-	/** A color to visualize background area for the combined data graph. */
-	FLinearColor ColorBackground;
-};
-
-
-
 /*-----------------------------------------------------------------------------
 	Declarations
 -----------------------------------------------------------------------------*/
@@ -282,14 +241,7 @@ protected:
 	FSelectionChangedForIndexEvent SelectionChangedForIndexEvent;
 
 public:
-	void AddInnerGraph
-	( 
-		const uint32 InStatID, 	
-		const FLinearColor InColorAverage,
-		const FLinearColor InColorExtremes,
-		const FLinearColor InColorBackground, 
-		const FCombinedGraphDataSourceRef& CombinedGraphDataSource  
-	);
+	void AddInnerGraph( const FTrackedStatPtr& TrackedStat );
 	void RemoveInnerGraph( const uint32 StatID );
 
 	/**
@@ -366,11 +318,7 @@ protected:
 	/** Handles FIsActionChecked for DataGraph_ViewMode_SetTimeBased. */
 	bool ViewMode_SetTimeBased_IsChecked() const;
 
-	const FGraphDescription* GetFirstGraph() const
-	{
-		const FGraphDescription* Result = StatIDToGraphDescriptionMapping.Num() > 0 ? &StatIDToGraphDescriptionMapping.CreateConstIterator().Value() : NULL;
-		return Result;
-	}
+	const FTrackedStatPtr GetFirstGraph() const;
 
 	/*-----------------------------------------------------------------------------
 		Misc
@@ -379,7 +327,7 @@ protected:
 	/** Horizontal box widget where graph descriptions are displayed. */
 	TSharedPtr<SVerticalBox> GraphDescriptionsVBox;
 
-	TMap< uint32, FGraphDescription > StatIDToGraphDescriptionMapping;
+	TMap< uint32, FTrackedStatPtr > StatIDToGraphDescriptionMapping;
 	TMap< uint32, TSharedRef<SWidget> > StatIDToWidgetMapping;
 
 	/** The current mouse position. */
diff --git a/Engine/Source/Developer/Profiler/Private/Widgets/SEventGraph.cpp b/Engine/Source/Developer/Profiler/Private/Widgets/SEventGraph.cpp
--- a/Engine/Source/Developer/Profiler/Private/Widgets/SEventGraph.cpp
+++ b/Engine/Source/Developer/Profiler/Private/Widgets/SEventGraph.cpp
@@ -58,7 +58,7 @@ namespace EventGraphColumns
 	// @TODO: Needs refactoring
 
 	/** Contains basic information about columns used in the event graph widget. Names should be localized. */
-	static const FEventGraphColumn Collection[ EEventPropertyIndex::None+1 ] =
+	static const FEventGraphColumn Collection[ 7 ] =
 	{
 		// Generated from a XLSX file.
 		FEventGraphColumn
@@ -94,39 +94,6 @@ namespace EventGraphColumns
 			64.0f
 		),
 		
-		FEventGraphColumn
-		(
-			EEventPropertyIndex::MinInclusiveTimeMS,
-			TEXT("min"),
-			LOCTEXT("MinInclusiveTimeMSTitle", "Min Inc Time (MS)"),
-			LOCTEXT("MinInclusiveTimeMSDesc", "Minimum inclusive time of all instances for this event, in milliseconds"),
-			true,false,true,true,false,
-			HAlign_Right,
-			0.0f
-		),
-		
-		FEventGraphColumn
-		(
-			EEventPropertyIndex::MaxInclusiveTimeMS,
-			TEXT("max"),
-			LOCTEXT("MaxInclusiveTimeMSTitle", "Max Inc Time (MS)"),
-			LOCTEXT("MaxInclusiveTimeMSDesc", "Maximum inclusive time of all instances for this event, in milliseconds"),
-			true,false,true,true,false,
-			HAlign_Right,
-			0.0f
-		),
-		
-		FEventGraphColumn
-		(
-			EEventPropertyIndex::AvgInclusiveTimeMS,
-			TEXT("avg"),
-			LOCTEXT("AverageInclusiveTimeMSTitle", "Avg Inc Time (MS)"),
-			LOCTEXT("AverageInclusiveTimeMSDesc", "Average inclusive time of all instances for this event, in milliseconds"),
-			true,false,true,true,false,
-			HAlign_Right,
-			0.0f
-		),
-		
 		FEventGraphColumn
 		(
 			EEventPropertyIndex::ExclusiveTimeMS,
@@ -149,17 +116,6 @@ namespace EventGraphColumns
 			64.0f
 		),
 		
-		FEventGraphColumn
-		(
-			EEventPropertyIndex::AvgInclusiveTimePerCallMS,
-			TEXT("avgper"),
-			LOCTEXT("AverageInclusiveTimePerCallMSTitle", "Avg Inc Time per Call (MS)"),
-			LOCTEXT("AverageInclusiveTimePerCallMSDesc", "Average inclusive time per call of all instances for this event, in milliseconds"),
-			true,false,true,false,false,
-			HAlign_Right,
-			0.0f
-		),
-		
 		FEventGraphColumn
 		(
 			EEventPropertyIndex::NumCallsPerFrame,
@@ -170,111 +126,6 @@ namespace EventGraphColumns
 			HAlign_Right,
 			64.0f
 		),
-		
-		FEventGraphColumn
-		(
-			EEventPropertyIndex::AvgNumCallsPerFrame,
-			TEXT("avgcalls"),
-			LOCTEXT("AverageCallsPerFrameTitle", "Avg Calls"),
-			LOCTEXT("AverageCallsPerFrameDesc", "Average number of times this event was called"),
-			true,true,true,true,false,
-			HAlign_Right,
-			64.0f
-		),
-		
-		// TODO: Show only in the tooltip
-		FEventGraphColumn
-		(
-			EEventPropertyIndex::ThreadName,
-			TEXT("th"),
-			LOCTEXT("ThreadNameTitle", "Thread Name"),
-			LOCTEXT("ThreadNameDesc", "Name of the thread that this event was captured on"),
-			true,false,true,false,false,
-			HAlign_Left,
-			0.0f
-		),
-		
-		// TODO: Show only in the tooltip
-		FEventGraphColumn
-		(
-			EEventPropertyIndex::ThreadDurationMS,
-			TEXT("td"),
-			LOCTEXT("ThreadDurationTitle", "Thread Duration (MS)"),
-			LOCTEXT("ThreadDurationDesc", "Duration of the thread that this event was captured on, in milliseconds"),
-			true,false,true,false,false,
-			HAlign_Right,
-			0.0f
-		),
-
-		// TODO: Show only in the tooltip
-		FEventGraphColumn
-		(
-			EEventPropertyIndex::FrameDurationMS,
-			TEXT("fd"),
-			LOCTEXT("FrameDurationMSTitle", "Frame Duration (MS)"),
-			LOCTEXT("FrameDurationMSDesc", "Duration of the frame that this event was captured on, in milliseconds"),
-			true,false,true,false,false,
-			HAlign_Right,
-			0.0f
-		),
-		
-		// TODO: Show only in the tooltip
-		FEventGraphColumn
-		(
-			EEventPropertyIndex::ThreadPct,
-			TEXT("t%"),
-			LOCTEXT("ThreadPercentageTitle", "% of Thread"),
-			LOCTEXT("ThreadPercentageDesc", "Percent of inclusive time spent by this event in the particular thread"),
-			true,true,true,false,false,
-			HAlign_Right,
-			64.0f
-		),
-		
-		// TODO: Show only in the tooltip
-		FEventGraphColumn
-		(
-			EEventPropertyIndex::FramePct,
-			TEXT("f%"),
-			LOCTEXT("FramePercentageTitle", "% of Frame"),
-			LOCTEXT("FramePercentageDesc", "Percent of inclusive time spent by this event in the particular frame"),
-			true,true,true,false,false,
-			HAlign_Right,
-			64.0f
-		),
-		
-		// TODO: Show only in the tooltip
-		FEventGraphColumn
-		(
-			EEventPropertyIndex::ThreadToFramePct,
-			TEXT("t2f%"),
-			LOCTEXT("ThreadToFramePercentageTitle", "Thread to Frame (%)"),
-			LOCTEXT("ThreadToFramePercentageDesc", "Percent of time spent in the thread in relation to the entire frame"),
-			true,false,true,false,false,
-			HAlign_Right,
-			64.0f
-		),
-		
-		FEventGraphColumn
-		(
-			EEventPropertyIndex::StartTimeMS,
-			TEXT("start"),
-			LOCTEXT("StartTimeMSTitle", "Start Time (MS)"),
-			LOCTEXT("StartTimeMSDesc", "Start time of this event, in milliseconds"),
-			true,false,true,false,false,
-			HAlign_Right,
-			0.0f
-		),
-		
-		FEventGraphColumn
-		(
-			EEventPropertyIndex::GroupName,
-			TEXT("grp"),
-			LOCTEXT("GroupNameTitle", "Group Name"),
-			LOCTEXT("GroupNameDesc", "Name of the stat group this event belongs to"),
-			true,false,true,false,false,
-			HAlign_Left,
-			0.0f
-		),
 
 		// Fake column used as a default column for NAME_None
 		FEventGraphColumn
@@ -296,22 +147,9 @@ namespace EventGraphColumns
 		.Add(TEXT("StatName"),&Collection[EEventPropertyIndex::StatName])
 		.Add(TEXT("InclusiveTimeMS"),&Collection[EEventPropertyIndex::InclusiveTimeMS])
 		.Add(TEXT("InclusiveTimePct"),&Collection[EEventPropertyIndex::InclusiveTimePct])
-		.Add(TEXT("MinInclusiveTimeMS"),&Collection[EEventPropertyIndex::MinInclusiveTimeMS])
-		.Add(TEXT("MaxInclusiveTimeMS"),&Collection[EEventPropertyIndex::MaxInclusiveTimeMS])
-		.Add(TEXT("AvgInclusiveTimeMS"),&Collection[EEventPropertyIndex::AvgInclusiveTimeMS])
 		.Add(TEXT("ExclusiveTimeMS"),&Collection[EEventPropertyIndex::ExclusiveTimeMS])
 		.Add(TEXT("ExclusiveTimePct"),&Collection[EEventPropertyIndex::ExclusiveTimePct])
-		.Add(TEXT("AvgInclusiveTimePerCallMS"),&Collection[EEventPropertyIndex::AvgInclusiveTimePerCallMS])
 		.Add(TEXT("NumCallsPerFrame"),&Collection[EEventPropertyIndex::NumCallsPerFrame])
-		.Add(TEXT("AvgNumCallsPerFrame"),&Collection[EEventPropertyIndex::AvgNumCallsPerFrame])
-		.Add(TEXT("ThreadName"),&Collection[EEventPropertyIndex::ThreadName])
-		.Add(TEXT("ThreadDurationMS"),&Collection[EEventPropertyIndex::ThreadDurationMS])
-		.Add(TEXT("FrameDurationMS"),&Collection[EEventPropertyIndex::FrameDurationMS])
-		.Add(TEXT("ThreadPct"),&Collection[EEventPropertyIndex::ThreadPct])
-		.Add(TEXT("FramePct"),&Collection[EEventPropertyIndex::FramePct])
-		.Add(TEXT("ThreadToFramePct"),&Collection[EEventPropertyIndex::ThreadToFramePct])
-		.Add(TEXT("StartTimeMS"),&Collection[EEventPropertyIndex::StartTimeMS])
-		.Add(TEXT("GroupName"),&Collection[EEventPropertyIndex::GroupName])
 		.Add(NAME_None,&Collection[EEventPropertyIndex::None])
 		;
 }
@@ -2505,6 +2343,7 @@ void SEventGraph::CreateEvents()
 		HierarchicalToFlatCoalesced.Add( EventPtr->_StatName, EventPtr );
 	}
 
+	// Should ignore recursion!
 	for( auto It = FlatIncCoalescedEvents.CreateConstIterator(); It; ++It )
 	{
 		const TArray<FEventGraphSamplePtr>& InclusiveCoalescedEvents = It.Value();
@@ -2512,7 +2351,7 @@ void SEventGraph::CreateEvents()
 		FEventGraphSamplePtr CoalescedEvent = InclusiveCoalescedEvents[0];
 		for( int32 EventIndex = 1; EventIndex < InclusiveCoalescedEvents.Num(); EventIndex++ )
 		{
-			*CoalescedEvent += *InclusiveCoalescedEvents[EventIndex];
+			CoalescedEvent->Combine( InclusiveCoalescedEvents[EventIndex] );
 		}
 
 		Events_FlatCoalesced.Add( CoalescedEvent );
