Copying //UE4/Dev-Core to //UE4/Main

diff --git a/Engine/Source/Developer/Profiler/Private/Widgets/SEventGraphTooltip.cpp b/Engine/Source/Developer/Profiler/Private/Widgets/SEventGraphTooltip.cpp
--- a/Engine/Source/Developer/Profiler/Private/Widgets/SEventGraphTooltip.cpp
+++ b/Engine/Source/Developer/Profiler/Private/Widgets/SEventGraphTooltip.cpp
@@ -96,9 +96,9 @@ TSharedPtr<SToolTip> SEventGraphTooltip::GetTableCellTooltip( const FEventGraphS
 					.Padding( 2.0f )
 					[
 						SNew(SGridPanel)
-
-						// First row 
+					
 						//-----------------------------------------------------------------------------
+						// First row 
 						// Thread: [_ThreadName]
 						+SGridPanel::Slot(0,0)
 						.Padding( 2.0f )
@@ -108,279 +108,291 @@ TSharedPtr<SToolTip> SEventGraphTooltip::GetTableCellTooltip( const FEventGraphS
 							.TextStyle( FEditorStyle::Get(), TEXT("Profiler.TooltipBold") )
 						]
 						+SGridPanel::Slot(1,0)
-							.Padding( 2.0f )
-							[
-								SNew(STextBlock)
-								.Text( FText::FromName(EventSample->_ThreadName ))
-								.TextStyle( FEditorStyle::Get(), TEXT("Profiler.Tooltip") )
-							]
+						.Padding( 2.0f )
+						[
+							SNew(STextBlock)
+							.Text( FText::FromName(EventSample->_ThreadName ))
+							.TextStyle( FEditorStyle::Get(), TEXT("Profiler.Tooltip") )
+						]
 
 						// Event: [_StatName]
 						+SGridPanel::Slot(2,0)
-							.Padding( 2.0f )
-							[
-								SNew(STextBlock)
-								.Text(LOCTEXT("TT_Event", "Event:"))
-								.TextStyle( FEditorStyle::Get(), TEXT("Profiler.TooltipBold") )
-							]
+						.Padding( 2.0f )
+						[
+							SNew(STextBlock)
+							.Text(LOCTEXT("TT_Event", "Event:"))
+							.TextStyle( FEditorStyle::Get(), TEXT("Profiler.TooltipBold") )
+						]
 						+SGridPanel::Slot(3,0)
-							.Padding( 2.0f )
-							[
-								SNew(STextBlock)
-								.Text( FText::FromString(EventSample->GetShortEventName()) )
-								.TextStyle( FEditorStyle::Get(), TEXT("Profiler.Tooltip") )
-							]
+						.Padding( 2.0f )
+						[
+							SNew(STextBlock)
+							.Text( FText::FromString(EventSample->GetShortEventName()) )
+							.TextStyle( FEditorStyle::Get(), TEXT("Profiler.Tooltip") )
+						]
 
 						// Group: [_GroupName]
 						+SGridPanel::Slot(4,0)
-							.Padding( 2.0f )
-							[
-								SNew(STextBlock)
-								.Text(LOCTEXT("TT_Group", "Group:"))
-								.TextStyle( FEditorStyle::Get(), TEXT("Profiler.TooltipBold") )
-							]
+						.Padding( 2.0f )
+						[
+							SNew(STextBlock)
+							.Text(LOCTEXT("TT_Group", "Group:"))
+							.TextStyle( FEditorStyle::Get(), TEXT("Profiler.TooltipBold") )
+						]
 						+SGridPanel::Slot(5,0)
-							.Padding( 2.0f )
-							[
-								SNew(STextBlock)
-								.Text( FText::FromName(EventSample->_GroupName) )
-								.TextStyle( FEditorStyle::Get(), TEXT("Profiler.Tooltip") )
-							]
+						.Padding( 2.0f )
+						[
+							SNew(STextBlock)
+							.Text( FText::FromName(EventSample->_GroupName) )
+							.TextStyle( FEditorStyle::Get(), TEXT("Profiler.Tooltip") )
+						]
 
 						//-----------------------------------------------------------------------------
 						// Second row
 						// Inclusive time: 
 						+SGridPanel::Slot(0,1)
-							.Padding( 2.0f )
-							[
-								SNew(STextBlock)
-								.Text(LOCTEXT("TT_InclusiveTime", "Inclusive time:"))
-								.TextStyle( FEditorStyle::Get(), TEXT("Profiler.TooltipBold") )
-							]
+						.Padding( 2.0f )
+						[
+							SNew(STextBlock)
+							.Text(LOCTEXT("TT_InclusiveTime", "Inclusive time:"))
+							.TextStyle( FEditorStyle::Get(), TEXT("Profiler.TooltipBold") )
+						]
 						// [InclusiveTimeMS] ms
 						+SGridPanel::Slot(1,1)
-							.Padding( 2.0f )
-							[
-								SNew(STextBlock)
-								.Text( FText::FromString(EventSample->GetFormattedValue(EEventPropertyIndex::InclusiveTimeMS)) )
-								.TextStyle( FEditorStyle::Get(), TEXT("Profiler.Tooltip") )
-							]
+						.Padding( 2.0f )
+						[
+							SNew(STextBlock)
+							.Text( FText::FromString(EventSample->GetFormattedValue(EEventPropertyIndex::InclusiveTimeMS)) )
+							.TextStyle( FEditorStyle::Get(), TEXT("Profiler.Tooltip") )
+						]
 						// ([_InclusiveTimePct] % of the caller)
 						+SGridPanel::Slot(2,1)
-							.Padding( 2.0f )
-							[
-								SNew(STextBlock)
-								.Text( InclusiveTimePctCaller )
-								.TextStyle( FEditorStyle::Get(), TEXT("Profiler.Tooltip") )
-							]
+						.Padding( 2.0f )
+						[
+							SNew(STextBlock)
+							.Text( InclusiveTimePctCaller )
+							.TextStyle( FEditorStyle::Get(), TEXT("Profiler.Tooltip") )
+						]
 						// Average inclusive time per call of all instances for this event, in milliseconds
 						+SGridPanel::Slot(3,1)
-							.Padding( 2.0f )
-							[
-								SNew(STextBlock)
-								.Text(LOCTEXT("TT_AvgIncTimePerCall", "Avg inc time per call:"))
-								.TextStyle( FEditorStyle::Get(), TEXT("Profiler.TooltipBold") )
-							]
-						// [_AvgInclusiveTimePerCallMS] ms
-						+SGridPanel::Slot(4,1)
-							.Padding( 2.0f )
-							[
-								SNew(STextBlock)
-								.Text( FText::FromString(EventSample->GetFormattedValue(EEventPropertyIndex::AvgInclusiveTimePerCallMS)) )
-								.TextStyle( FEditorStyle::Get(), TEXT("Profiler.Tooltip") )
-							]
-
-						// Third row
+						.Padding( 2.0f )
+						[
+							SNew(STextBlock)
+							.Text(LOCTEXT("TT_AvgIncTimePerCall", "Avg inc time per call:"))
+							.TextStyle( FEditorStyle::Get(), TEXT("Profiler.TooltipBold") )
+						]
+						
 						//-----------------------------------------------------------------------------
+						// Third row
 						// 		_ThreadDurationMS, (_ThreadPct); /** Percent of inclusive time spent by this event in the particular thread. */
 						// 		_FrameDurationMS, (_FramePct); /** Percent of inclusive time spent by this event in the particular frame. */
 						+SGridPanel::Slot(0,2)
-							.Padding( 2.0f )
-							[
-								SNew(STextBlock)
-								.Text(LOCTEXT("TT_PctOfThread", "% of thread:"))
-								.TextStyle( FEditorStyle::Get(), TEXT("Profiler.TooltipBold") )
-							]
+						.Padding( 2.0f )
+						[
+							SNew(STextBlock)
+							.Text(LOCTEXT("TT_PctOfThread", "% of thread:"))
+							.TextStyle( FEditorStyle::Get(), TEXT("Profiler.TooltipBold") )
+						]
 						// [ThreadPct] %
 						+SGridPanel::Slot(1,2)
-							.Padding( 2.0f )
-							[
-								SNew(STextBlock)
-								.Text( FText::FromString(EventSample->GetFormattedValue(EEventPropertyIndex::ThreadPct)) )
-								.TextStyle( FEditorStyle::Get(), TEXT("Profiler.Tooltip") )
-							]
+						.Padding( 2.0f )
+						[
+							SNew(STextBlock)
+							.Text( FText::FromString(EventSample->GetFormattedValue(EEventPropertyIndex::ThreadPct)) )
+							.TextStyle( FEditorStyle::Get(), TEXT("Profiler.Tooltip") )
+						]
 
 						+SGridPanel::Slot(3,2)
-							.Padding( 2.0f )
-							[
-								SNew(STextBlock)
-								.Text(LOCTEXT("TT_PctOfFrame", "% of frame:"))
-								.TextStyle( FEditorStyle::Get(), TEXT("Profiler.TooltipBold") )
-							]
+						.Padding( 2.0f )
+						[
+							SNew(STextBlock)
+							.Text(LOCTEXT("TT_PctOfFrame", "% of frame:"))
+							.TextStyle( FEditorStyle::Get(), TEXT("Profiler.TooltipBold") )
+						]
 						// [FramePct] %
 						+SGridPanel::Slot(4,2)
-							.Padding( 2.0f )
-							[
-								SNew(STextBlock)
-								.Text( FText::FromString(EventSample->GetFormattedValue(EEventPropertyIndex::FramePct)) )
-								.TextStyle( FEditorStyle::Get(), TEXT("Profiler.Tooltip") )
-							]
+						.Padding( 2.0f )
+						[
+							SNew(STextBlock)
+							.Text( FText::FromString(EventSample->GetFormattedValue(EEventPropertyIndex::FramePct)) )
+							.TextStyle( FEditorStyle::Get(), TEXT("Profiler.Tooltip") )
+						]
 
 						//-----------------------------------------------------------------------------
-						// Third row
+						// Fourth row
 						// _MinInclusiveTimeMS, _MaxInclusiveTimeMS, _AvgInclusiveTimeMS;
 						// Min inclusive time: 
 						+SGridPanel::Slot(0,3)
-							.Padding( 2.0f )
-							[
-								SNew(STextBlock)
-								.Text(LOCTEXT("TT_MinMaxAvgIncTime", "Min/Max/Avg inclusive time:"))
-								.TextStyle( FEditorStyle::Get(), TEXT("Profiler.TooltipBold") )
-							]
+						.Padding( 2.0f )
+						[
+							SNew(STextBlock)
+							.Text(LOCTEXT("TT_MinMaxAvgIncTime", "Min/Max/Avg inclusive time:"))
+							.TextStyle( FEditorStyle::Get(), TEXT("Profiler.TooltipBold") )
+						]
 						// [_MinInclusiveTimeMS] ms
 						+SGridPanel::Slot(1,3)
-							.Padding( 2.0f )
-							[
-								SNew(STextBlock)
-								.Text( FText::FromString(EventSample->GetFormattedValue(EEventPropertyIndex::MinInclusiveTimeMS)) )
-								.TextStyle( FEditorStyle::Get(), TEXT("Profiler.Tooltip") )
-							]
+						.Padding( 2.0f )
+						[
+							SNew(STextBlock)
+							.Text( FText::FromString(EventSample->GetFormattedValue(EEventPropertyIndex::MinInclusiveTimeMS)) )
+							.TextStyle( FEditorStyle::Get(), TEXT("Profiler.Tooltip") )
+						]
 						// [_MaxInclusiveTimeMS] ms
 						+SGridPanel::Slot(2,3)
-							.Padding( 2.0f )
-							[
-								SNew(STextBlock)
-								.Text( FText::FromString(EventSample->GetFormattedValue(EEventPropertyIndex::MaxInclusiveTimeMS)) )
-								.TextStyle( FEditorStyle::Get(), TEXT("Profiler.Tooltip") )
-							]
+						.Padding( 2.0f )
+						[
+							SNew(STextBlock)
+							.Text( FText::FromString(EventSample->GetFormattedValue(EEventPropertyIndex::MaxInclusiveTimeMS)) )
+							.TextStyle( FEditorStyle::Get(), TEXT("Profiler.Tooltip") )
+						]
 						// [_AvgInclusiveTimeMS] ms
 						+SGridPanel::Slot(3,3)
-							.Padding( 2.0f )
-							[
-								SNew(STextBlock)
-								.Text( FText::FromString(EventSample->GetFormattedValue(EEventPropertyIndex::AvgInclusiveTimeMS)) )
-								.TextStyle( FEditorStyle::Get(), TEXT("Profiler.Tooltip") )
-							]
+						.Padding( 2.0f )
+						[
+							SNew(STextBlock)
+							.Text( FText::FromString(EventSample->GetFormattedValue(EEventPropertyIndex::AvgInclusiveTimeMS)) )
+							.TextStyle( FEditorStyle::Get(), TEXT("Profiler.Tooltip") )
+						]
 
 						//-----------------------------------------------------------------------------
-						// Fourth row
+						// Fifth row
 						// _ExclusiveTimeMS, _ExclusiveTimePct;
 						// Inclusive time: 
 						+SGridPanel::Slot(0,4)
-							.Padding( 2.0f )
-							[
-								SNew(STextBlock)
-								.Text(LOCTEXT("TT_ExclusiveTime", "Exclusive time:"))
-								.TextStyle( FEditorStyle::Get(), TEXT("Profiler.TooltipBold") )
-							]
+						.Padding( 2.0f )
+						[
+							SNew(STextBlock)
+							.Text(LOCTEXT("TT_ExclusiveTime", "Exclusive time:"))
+							.TextStyle( FEditorStyle::Get(), TEXT("Profiler.TooltipBold") )
+						]
 						// [_ExclusiveTimeMS] ms
 						+SGridPanel::Slot(1,4)
-							.Padding( 2.0f )
-							[
-								SNew(STextBlock)
-								.Text( FText::FromString(EventSample->GetFormattedValue(EEventPropertyIndex::ExclusiveTimeMS)) )
-								.TextStyle( FEditorStyle::Get(), TEXT("Profiler.Tooltip") )
-							]
+						.Padding( 2.0f )
+						[
+							SNew(STextBlock)
+							.Text( FText::FromString(EventSample->GetFormattedValue(EEventPropertyIndex::ExclusiveTimeMS)) )
+							.TextStyle( FEditorStyle::Get(), TEXT("Profiler.Tooltip") )
+						]
 						// ([_ExclusiveTimePct] % of this call's inc time)
 						+SGridPanel::Slot(2,4)
-							.Padding( 2.0f )
-							[
-								SNew(STextBlock)
-								.Text( ExclusiveTimePctCaller )
-								.TextStyle( FEditorStyle::Get(), TEXT("Profiler.Tooltip") )
-							]
+						.Padding( 2.0f )
+						[
+							SNew(STextBlock)
+							.Text( ExclusiveTimePctCaller )
+							.TextStyle( FEditorStyle::Get(), TEXT("Profiler.Tooltip") )
+						]
 
 						//-----------------------------------------------------------------------------
-						// Fifth row
+						// Sixth row
 						// _NumCallsPerFrame, _AvgNumCallsPerFrame;
 						// Inclusive time: 
 						+SGridPanel::Slot(0,5)
-							.Padding( 2.0f )
-							[
-								SNew(STextBlock)
-								.Text(LOCTEXT("TT_NulCallsPerFrame", "Num calls per frame:"))
-								.TextStyle( FEditorStyle::Get(), TEXT("Profiler.TooltipBold") )
-							]
+						.Padding( 2.0f )
+						[
+							SNew(STextBlock)
+							.Text(LOCTEXT("TT_NulCallsPerFrame", "Num calls per frame:"))
+							.TextStyle( FEditorStyle::Get(), TEXT("Profiler.TooltipBold") )
+						]
 						// [_NumCallsPerFrame]
 						+SGridPanel::Slot(1,5)
-							.Padding( 2.0f )
-							[
-								SNew(STextBlock)
-								.Text( FText::FromString(EventSample->GetFormattedValue(EEventPropertyIndex::NumCallsPerFrame)) )
-								.TextStyle( FEditorStyle::Get(), TEXT("Profiler.Tooltip") )
-							]
+						.Padding( 2.0f )
+						[
+							SNew(STextBlock)
+							.Text( FText::FromString(EventSample->GetFormattedValue(EEventPropertyIndex::NumCallsPerFrame)) )
+							.TextStyle( FEditorStyle::Get(), TEXT("Profiler.Tooltip") )
+						]
 
-						// Inclusive time: 
-						+SGridPanel::Slot(2,5)
-							.Padding( 2.0f )
-							[
-								SNew(STextBlock)
-								.Text(LOCTEXT("TT_AvgNumCallsPerFrame", "Avg num calls per frame:"))
-								.TextStyle( FEditorStyle::Get(), TEXT("Profiler.TooltipBold") )
-							]
-						// [_AvgNumCallsPerFrame]
-						+SGridPanel::Slot(3,5)
-							.Padding( 2.0f )
-							[
-								SNew(STextBlock)
-								.Text( FText::FromString(EventSample->GetFormattedValue(EEventPropertyIndex::AvgNumCallsPerFrame)) )
-								.TextStyle( FEditorStyle::Get(), TEXT("Profiler.Tooltip") )
-							]
 
 						//-----------------------------------------------------------------------------
-						// Sixth and seventh row
+						// Seventh row
+						// _MinNumCallsPerFrame, _MaxNumCallsPerFrame, _AvgNumCallsPerFrame;
+						// Min inclusive time: 
+						+SGridPanel::Slot( 0, 6 )
+						.Padding( 2.0f )
+						[
+							SNew( STextBlock )
+							.Text( LOCTEXT( "TT_MinMaxAvgNumCallsPerFrame", "Min/Max/Avg calls per frame:" ) )
+							.TextStyle( FEditorStyle::Get(), TEXT( "Profiler.TooltipBold" ) )
+						]
+						// [_MinNumCallsPerFrame]
+						+ SGridPanel::Slot( 1, 6 )
+						.Padding( 2.0f )
+						[
+							SNew( STextBlock )
+							.Text( FText::FromString( EventSample->GetFormattedValue( EEventPropertyIndex::MinNumCallsPerFrame ) ) )
+							.TextStyle( FEditorStyle::Get(), TEXT( "Profiler.Tooltip" ) )
+						]
+						// [_MaxNumCallsPerFrame]
+						+ SGridPanel::Slot( 2, 6 )
+						.Padding( 2.0f )
+						[
+							SNew( STextBlock )
+							.Text( FText::FromString( EventSample->GetFormattedValue( EEventPropertyIndex::MaxNumCallsPerFrame ) ) )
+							.TextStyle( FEditorStyle::Get(), TEXT( "Profiler.Tooltip" ) )
+						]
+						// [_AvgNumCallsPerFrame]
+						+ SGridPanel::Slot( 3, 6 )
+						.Padding( 2.0f )
+						[
+							SNew( STextBlock )
+							.Text( FText::FromString( EventSample->GetFormattedValue( EEventPropertyIndex::AvgNumCallsPerFrame ) ) )
+							.TextStyle( FEditorStyle::Get(), TEXT( "Profiler.Tooltip" ) )
+						]
+						
+						//-----------------------------------------------------------------------------
+						// Eighth row
 						// _ThreadDurationMS, _FrameDurationMS, _ThreadToFramePct
 						// Thread duration: 
-						+SGridPanel::Slot(0,6)
-							.Padding( 2.0f )
-							[
-								SNew(STextBlock)
-								.Text(LOCTEXT("TT_ThreadDuration", "Thread duration:"))
-								.TextStyle( FEditorStyle::Get(), TEXT("Profiler.TooltipBold") )
-							]
+						+SGridPanel::Slot(0,7)
+						.Padding( 2.0f )
+						[
+							SNew(STextBlock)
+							.Text(LOCTEXT("TT_ThreadDuration", "Thread duration:"))
+							.TextStyle( FEditorStyle::Get(), TEXT("Profiler.TooltipBold") )
+						]
 						// [_ThreadDurationMS]
-						+SGridPanel::Slot(1,6)
-							.Padding( 2.0f )
-							[
-								SNew(STextBlock)
-								.Text( FText::FromString(EventSample->GetFormattedValue(EEventPropertyIndex::ThreadDurationMS)) )
-								.TextStyle( FEditorStyle::Get(), TEXT("Profiler.Tooltip") )
-							]
+						+SGridPanel::Slot(1,7)
+						.Padding( 2.0f )
+						[
+							SNew(STextBlock)
+							.Text( FText::FromString(EventSample->GetFormattedValue(EEventPropertyIndex::ThreadDurationMS)) )
+							.TextStyle( FEditorStyle::Get(), TEXT("Profiler.Tooltip") )
+						]
 
 						// Frame duration: 
-						+SGridPanel::Slot(2,6)
-							.Padding( 2.0f )
-							[
-								SNew(STextBlock)
-								.Text(LOCTEXT("TT_FrameDuration", "Frame duration:"))
-								.TextStyle( FEditorStyle::Get(), TEXT("Profiler.TooltipBold") )
-							]
+						+SGridPanel::Slot(2,7)
+						.Padding( 2.0f )
+						[
+							SNew(STextBlock)
+							.Text(LOCTEXT("TT_FrameDuration", "Frame duration:"))
+							.TextStyle( FEditorStyle::Get(), TEXT("Profiler.TooltipBold") )
+						]
 						// [_FrameDurationMS]
-						+SGridPanel::Slot(3,6)
-							.Padding( 2.0f )
-							[
-								SNew(STextBlock)
-								.Text( FText::FromString(EventSample->GetFormattedValue(EEventPropertyIndex::FrameDurationMS)) )
-								.TextStyle( FEditorStyle::Get(), TEXT("Profiler.Tooltip") )
-							]
+						+SGridPanel::Slot(3,7)
+						.Padding( 2.0f )
+						[
+							SNew(STextBlock)
+							.Text( FText::FromString(EventSample->GetFormattedValue(EEventPropertyIndex::FrameDurationMS)) )
+							.TextStyle( FEditorStyle::Get(), TEXT("Profiler.Tooltip") )
+						]
 
 						/** Percent of time spent in the thread in relation to the entire frame. */
 						// Thread -> Frame: 
-						+SGridPanel::Slot(4,6)
-							.Padding( 2.0f )
-							[
-								SNew(STextBlock)
-								.Text(LOCTEXT("ThreadToFrame", "Thread to Frame:"))
-								.TextStyle( FEditorStyle::Get(), TEXT("Profiler.TooltipBold") )
-							]
+						+SGridPanel::Slot(4,7)
+						.Padding( 2.0f )
+						[
+							SNew(STextBlock)
+							.Text(LOCTEXT("ThreadToFrame", "Thread to Frame:"))
+							.TextStyle( FEditorStyle::Get(), TEXT("Profiler.TooltipBold") )
+						]
 						// [_FrameDurationMS]
-						+SGridPanel::Slot(5,6)
-							.Padding( 2.0f )
-							[
-								SNew(STextBlock)
-								.Text( FText::FromString(EventSample->GetFormattedValue(EEventPropertyIndex::ThreadToFramePct)) )
-								.TextStyle( FEditorStyle::Get(), TEXT("Profiler.Tooltip") )
-							]
+						+SGridPanel::Slot(5,7)
+						.Padding( 2.0f )
+						[
+							SNew(STextBlock)
+							.Text( FText::FromString(EventSample->GetFormattedValue(EEventPropertyIndex::ThreadToFramePct)) )
+							.TextStyle( FEditorStyle::Get(), TEXT("Profiler.Tooltip") )
+						]
 					]
 
 				+SVerticalBox::Slot()
diff --git a/Engine/Source/Developer/Profiler/Private/Widgets/SProfilerGraphPanel.cpp b/Engine/Source/Developer/Profiler/Private/Widgets/SProfilerGraphPanel.cpp
--- a/Engine/Source/Developer/Profiler/Private/Widgets/SProfilerGraphPanel.cpp
+++ b/Engine/Source/Developer/Profiler/Private/Widgets/SProfilerGraphPanel.cpp
@@ -146,15 +146,15 @@ void SProfilerGraphPanel::VerticalScrollBar_OnUserScrolled( float ScrollOffset )
 }
 
 
-void SProfilerGraphPanel::ProfilerManager_OnTrackedStatChanged( const FTrackedStat& TrackedStat, bool bIsTracked )
+void SProfilerGraphPanel::ProfilerManager_OnTrackedStatChanged( const FTrackedStatPtr& TrackedStat, bool bIsTracked )
 {
 	if( bIsTracked )
 	{
-		DataGraph->AddInnerGraph( TrackedStat.StatID, TrackedStat.ColorAverage, TrackedStat.ColorBackground, TrackedStat.ColorExtremes, TrackedStat.CombinedGraphDataSource );
+		DataGraph->AddInnerGraph( TrackedStat );
 	}
 	else
 	{
-		DataGraph->RemoveInnerGraph( TrackedStat.StatID );
+		DataGraph->RemoveInnerGraph( TrackedStat->StatID );
 	}
 }
 
diff --git a/Engine/Source/Developer/Profiler/Private/Widgets/SProfilerGraphPanel.h b/Engine/Source/Developer/Profiler/Private/Widgets/SProfilerGraphPanel.h
--- a/Engine/Source/Developer/Profiler/Private/Widgets/SProfilerGraphPanel.h
+++ b/Engine/Source/Developer/Profiler/Private/Widgets/SProfilerGraphPanel.h
@@ -33,7 +33,7 @@ public:
 
 protected:
 	/** Called when the status of specified tracked stat has changed. */
-	void ProfilerManager_OnTrackedStatChanged( const FTrackedStat& TrackedStat, bool bIsTracked );
+	void ProfilerManager_OnTrackedStatChanged( const FTrackedStatPtr& TrackedStat, bool bIsTracked );
 
 	/**
 	 * Called when the user scrolls the scrollbar
diff --git a/Engine/Source/Developer/Profiler/Public/ProfilerCommon.h b/Engine/Source/Developer/Profiler/Public/ProfilerCommon.h
--- a/Engine/Source/Developer/Profiler/Public/ProfilerCommon.h
+++ b/Engine/Source/Developer/Profiler/Public/ProfilerCommon.h
@@ -17,10 +17,6 @@ typedef TSharedPtr<class FGraphDataSource> FGraphDataSourcePtr;
 /** Type definition for shared references to const instances of FGraphDataSource. */
 typedef TSharedRef<const class FGraphDataSource> FGraphDataSourceRefConst;
 
-/** Type definition for weak references to instances of FGraphDataSource. */
-typedef TWeakPtr<class FGraphDataSource> FGraphDataSourceWeak;
-
-
 /** Type definition for shared pointers to instances of FCombinedGraphDataSource. */
 typedef TSharedPtr<class FCombinedGraphDataSource> FCombinedGraphDataSourcePtr;
 
@@ -51,6 +47,9 @@ typedef TSharedPtr<class FProfilerStatMetaData> FProfilerStatMetaDataPtr;
 /** Type definition for shared references to instances of FProfilerStatMetaData. */
 typedef TSharedRef<class FProfilerStatMetaData> FProfilerStatMetaDataRef;
 
+/** Type definition for shared pointers to instances of FTrackedStat. */
+typedef TSharedPtr<class FTrackedStat> FTrackedStatPtr;
+
 class FProfilerStat;
 class FProfilerGroup;
 class FProfilerStatMetaData;
diff --git a/Engine/Source/Developer/ProfilerClient/Private/ProfilerClientManager.cpp b/Engine/Source/Developer/ProfilerClient/Private/ProfilerClientManager.cpp
--- a/Engine/Source/Developer/ProfilerClient/Private/ProfilerClientManager.cpp
+++ b/Engine/Source/Developer/ProfilerClient/Private/ProfilerClientManager.cpp
@@ -403,6 +403,11 @@ FServiceConnection::~FServiceConnection()
 		delete StatsReader;
 		StatsReader = nullptr;
 	}
+
+	for (const auto& It : ReceivedData)
+	{
+		delete It.Value;
+	}
 }
 
 void FServiceConnection::Initialize( const FProfilerServiceAuthorize& Message, const IMessageContextRef& Context )
@@ -574,13 +579,12 @@ bool FProfilerClientManager::HandleTicker( float DeltaTime )
 bool FProfilerClientManager::HandleMessagesTicker( float DeltaTime )
 {
 #if STATS
-	// #YRX_tag 2015-11-26 rework 
 	for (auto It = Connections.CreateIterator(); It; ++It)
 	{
 		FServiceConnection& Connection = It.Value();
 
 		TArray<int64> Frames;
-		Connection.PendingMessages.GenerateKeyArray( Frames );
+		Connection.ReceivedData.GenerateKeyArray( Frames );
 		Frames.Sort();
 
 		// MessageBus sends all data in out of order fashion.
@@ -589,40 +593,32 @@ bool FProfilerClientManager::HandleMessagesTicker( float DeltaTime )
 
 		for( int32 Index = 0; Index < Frames.Num(); Index++ )
 		{
-			if (Connection.PendingMessages.Num() < NUM_BUFFERED_FRAMES)
+			if (Connection.ReceivedData.Num() < NUM_BUFFERED_FRAMES)
 			{
 				break;
 			}
 
 			FScopeLogTime SLT( "HandleMessagesTicker" );
 
-			int64 FrameNum = Frames[Index];
-			TArray<uint8>& Data = *Connection.PendingMessages.Find( FrameNum );
-
-			// Pass the data to the visualization code.
-			FMemoryReader MemoryReader( Data );
-
-			int64 Size = MemoryReader.TotalSize();
+			const int64 FrameNum = Frames[Index];
+			const TArray<uint8>* const Data = Connection.ReceivedData.FindChecked( FrameNum );
 			FStatsReadStream& Stream = Connection.Stream;
 
-			UE_LOG( LogProfilerClient, VeryVerbose, TEXT( "Frame=%i/%i, FNamesIndexMap=%i, CurrentMetadataSize=%i" ), FrameNum, Frames.Num(), Connection.Stream.FNamesIndexMap.Num(), Connection.CurrentThreadState.ShortNameToLongName.Num() );
 
-			// read in the data and post if we reach a 
+			// Read all messages from the uncompressed buffer.
+			FMemoryReader MemoryReader( *Data, true );
+			while (MemoryReader.Tell() < MemoryReader.TotalSize())
 			{
-				SCOPE_CYCLE_COUNTER(STAT_PC_ReadStatMessages);
-				while(MemoryReader.Tell() < Size)
-				{
-					// read the message
-					FStatMessage Message(Stream.ReadMessage(MemoryReader));
+				// Read the message.
+				FStatMessage Message( Stream.ReadMessage( MemoryReader ) );
+				new (Connection.PendingStatMessagesMessages)FStatMessage( Message );
+			}
 
-					if (Message.NameAndInfo.GetField<EStatOperation>() == EStatOperation::AdvanceFrameEventGameThread)
-					{
-						Connection.AddCollectedStatMessages( Message );
-					}
+			// Adds a new from from the pending messages, the pending messages will be removed after the call.
+			Connection.CurrentThreadState.ProcessMetaDataAndLeaveDataOnly( Connection.PendingStatMessagesMessages );
+			Connection.CurrentThreadState.AddFrameFromCondensedMessages( Connection.PendingStatMessagesMessages );
 
-					new (Connection.PendingStatMessagesMessages) FStatMessage(Message);
-				}
-			}
+			UE_LOG( LogProfilerClient, VeryVerbose, TEXT( "Frame=%i/%i, FNamesIndexMap=%i, CurrentMetadataSize=%i" ), FrameNum, Frames.Num(), Connection.Stream.FNamesIndexMap.Num(), Connection.CurrentThreadState.ShortNameToLongName.Num() );
 
 			// create an old format data frame from the data
 			Connection.GenerateProfilerDataFrame();
@@ -636,7 +632,8 @@ bool FProfilerClientManager::HandleMessagesTicker( float DeltaTime )
 			// send the data out
 			ProfilerDataDelegate.Broadcast( Connection.InstanceId, Connection.CurrentData );
 
-			Connection.PendingMessages.Remove( FrameNum );
+			delete Data;
+			Connection.ReceivedData.Remove( FrameNum );
 		}
 	}
 
@@ -733,10 +730,8 @@ void FProfilerClientManager::SendDataToGame( TArray<uint8>* DataToGame, int64 Fr
 
 		// Add the message to the connections queue.
 		UE_LOG( LogProfilerClient, VeryVerbose, TEXT( "Frame: %i, UncompressedSize: %i, InstanceId: %s" ), Frame, DataToGame->Num(), *InstanceId.ToString() );
-		Connection.PendingMessages.Add( Frame, MoveTemp( *DataToGame ) );
+		Connection.ReceivedData.Add( Frame, DataToGame );
 	}
-
-	delete DataToGame;
 }
 
 // #YRX_tag 2015-11-19 Convert to thread safe pointers?
@@ -787,47 +782,6 @@ void FProfilerClientManager::FinalizeLoading()
 }
 
 #if STATS
-void FServiceConnection::UpdateMetaData_DISABLED()
-{
-// 	// loop through the stats meta data messages
-// 	for (auto It = CurrentThreadState.ShortNameToLongName.CreateConstIterator(); It; ++It)
-// 	{
-// 		FStatMessage const& LongName = It.Value();
-// 		const FName GroupName = LongName.NameAndInfo.GetGroupName();
-// 
-// 		uint32 StatType = STATTYPE_Error;
-// 		if (LongName.NameAndInfo.GetField<EStatDataType>() == EStatDataType::ST_int64)
-// 		{
-// 			if( LongName.NameAndInfo.GetFlag( EStatMetaFlags::IsCycle ) )
-// 			{
-// 				StatType = STATTYPE_CycleCounter;
-// 			}
-// 			else if( LongName.NameAndInfo.GetFlag( EStatMetaFlags::IsMemory ) )
-// 			{
-// 				StatType = STATTYPE_MemoryCounter;
-// 			}
-// 			else
-// 			{
-// 				StatType = STATTYPE_AccumulatorDWORD;
-// 			}
-// 		}
-// 		else if (LongName.NameAndInfo.GetField<EStatDataType>() == EStatDataType::ST_double)
-// 		{
-// 			StatType = STATTYPE_AccumulatorFLOAT;
-// 		}
-// 		if (StatType != STATTYPE_Error)
-// 		{
-// 			FindOrAddStat(LongName.NameAndInfo, StatType);
-// 		}
-// 
-// 		// Threads metadata.
-// 		const bool bIsThread = FStatConstants::NAME_ThreadGroup == GroupName;
-// 		if( bIsThread )
-// 		{
-// 			FindOrAddThread( LongName.NameAndInfo );
-// 		}	
-// 	}
-}
 
 int32 FServiceConnection::FindOrAddStat( const FStatNameAndInfo& StatNameAndInfo, uint32 StatType)
 {
@@ -1059,14 +1013,6 @@ void FServiceConnection::GenerateCycleGraphs(const FRawStatStackNode& Root, TMap
 	}
 }
 
-void FServiceConnection::AddCollectedStatMessages( FStatMessage Message )
-{
-	SCOPE_CYCLE_COUNTER( STAT_PC_AddStatMessages );
-	new (PendingStatMessagesMessages)FStatMessage( Message );
-	CurrentThreadState.AddCondensedMessages( PendingStatMessagesMessages );
-	PendingStatMessagesMessages.Reset();
-}
-
 void FServiceConnection::GenerateProfilerDataFrame()
 {
 	SCOPE_CYCLE_COUNTER( STAT_PC_GenerateDataFrame );
diff --git a/Engine/Source/Developer/ProfilerClient/Private/ProfilerClientManager.h b/Engine/Source/Developer/ProfilerClient/Private/ProfilerClientManager.h
--- a/Engine/Source/Developer/ProfilerClient/Private/ProfilerClientManager.h
+++ b/Engine/Source/Developer/ProfilerClient/Private/ProfilerClientManager.h
@@ -31,7 +31,7 @@ struct FServiceConnection
 
 #if STATS
 	/** Current stats data */
-	FStatsThreadState CurrentThreadState;
+	FStatsLoadedState CurrentThreadState;
 #endif
 
 	/** Provides an FName to GroupId mapping */
@@ -52,7 +52,7 @@ struct FServiceConnection
 	TArray<FProfilerDataFrame> DataFrames;
 
 	/** Messages received and pending process, they are stored in a map as we can receive them out of order */
-	TMap<int64, TArray<uint8> > PendingMessages;
+	TMap<int64, TArray<uint8>* > ReceivedData;
 
 #if STATS
 	// generates the old style cycle graph
@@ -70,12 +70,6 @@ struct FServiceConnection
 	// adds a new style stat FName to the list of threads and generates an old style id and description
 	int32 FindOrAddThread(const FStatNameAndInfo& Thread);
 
-	// updates the meta date for the given instance
-	void UpdateMetaData_DISABLED();
-
-	/** Adds all collected stat messages to the current stats thread state. */
-	void AddCollectedStatMessages( FStatMessage Message );
-
 	/** Generates profiler data frame based on the collected stat messages. */
 	void GenerateProfilerDataFrame();
 #endif
diff --git a/Engine/Source/Developer/ProfilerService/Private/ProfilerServiceManager.cpp b/Engine/Source/Developer/ProfilerService/Private/ProfilerServiceManager.cpp
--- a/Engine/Source/Developer/ProfilerService/Private/ProfilerServiceManager.cpp
+++ b/Engine/Source/Developer/ProfilerService/Private/ProfilerServiceManager.cpp
@@ -327,6 +327,7 @@ void FProfilerServiceManager::HandleServiceSubscribeMessage( const FProfilerServ
 
 		// Send authorize.
 		MessageEndpoint->Send( new FProfilerServiceAuthorize( SessionId, InstanceId ), SenderAddress );
+		// Eventually send the metadata if needed.
 
 		// Initiate the ping callback
 		if (ClientData.Num() == 1)
diff --git a/Engine/Source/Editor/SceneOutliner/Private/SceneOutlinerActorInfoColumn.cpp b/Engine/Source/Editor/SceneOutliner/Private/SceneOutlinerActorInfoColumn.cpp
--- a/Engine/Source/Editor/SceneOutliner/Private/SceneOutlinerActorInfoColumn.cpp
+++ b/Engine/Source/Editor/SceneOutliner/Private/SceneOutlinerActorInfoColumn.cpp
@@ -109,9 +109,9 @@ SHeaderRow::FColumn::FArguments FActorInfoColumn::ConstructHeaderRowColumn()
 {
 	if( ModeOptions.Num() == 0 )
 	{
-		for( auto CurModeIndex = 0; CurModeIndex < ECustomColumnMode::Count; ++CurModeIndex )
+		for( ECustomColumnMode::Type CurMode : TEnumRange<ECustomColumnMode::Type>() )
 		{
-			ModeOptions.Add( MakeShareable( new ECustomColumnMode::Type( ( ECustomColumnMode::Type )CurModeIndex ) ) );
+			ModeOptions.Add( MakeShareable( new ECustomColumnMode::Type( CurMode ) ) );
 		}
 	}
 
diff --git a/Engine/Source/Editor/SceneOutliner/Private/SceneOutlinerActorInfoColumn.h b/Engine/Source/Editor/SceneOutliner/Private/SceneOutlinerActorInfoColumn.h
--- a/Engine/Source/Editor/SceneOutliner/Private/SceneOutlinerActorInfoColumn.h
+++ b/Engine/Source/Editor/SceneOutliner/Private/SceneOutlinerActorInfoColumn.h
@@ -111,4 +111,6 @@ private:
 };
 
 
-}	// namespace SceneOutliner
\ No newline at end of file
+}	// namespace SceneOutliner
+
+ENUM_RANGE_BY_COUNT(SceneOutliner::ECustomColumnMode::Type, SceneOutliner::ECustomColumnMode::Count)
diff --git a/Engine/Source/Editor/UnrealEd/Private/Commandlets/CookCommandlet.cpp b/Engine/Source/Editor/UnrealEd/Private/Commandlets/CookCommandlet.cpp
--- a/Engine/Source/Editor/UnrealEd/Private/Commandlets/CookCommandlet.cpp
+++ b/Engine/Source/Editor/UnrealEd/Private/Commandlets/CookCommandlet.cpp
@@ -93,7 +93,7 @@ bool UCookCommandlet::CookOnTheFly( FGuid InstanceId, int32 Timeout, bool bForce
 	CookFlags |= bCompressed ? ECookInitializationFlags::Compressed : ECookInitializationFlags::None;
 	CookFlags |= bIterativeCooking ? ECookInitializationFlags::Iterative : ECookInitializationFlags::None;
 	CookFlags |= bSkipEditorContent ? ECookInitializationFlags::SkipEditorContent : ECookInitializationFlags::None;
-
+	CookFlags |= bUnversioned ? ECookInitializationFlags::Unversioned : ECookInitializationFlags::None;
 
 	CookOnTheFlyServer->Initialize( ECookMode::CookOnTheFly, CookFlags );
 
@@ -1024,6 +1024,7 @@ bool UCookCommandlet::NewCook( const TArray<ITargetPlatform*>& Platforms, TArray
 	CookFlags |= bIterativeCooking ? ECookInitializationFlags::Iterative : ECookInitializationFlags::None;
 	CookFlags |= bSkipEditorContent ? ECookInitializationFlags::SkipEditorContent : ECookInitializationFlags::None;	
 	CookFlags |= bUseSerializationForGeneratingPackageDependencies ? ECookInitializationFlags::UseSerializationForPackageDependencies : ECookInitializationFlags::None;
+	CookFlags |= bUnversioned ? ECookInitializationFlags::Unversioned : ECookInitializationFlags::None;
 
 	TArray<UClass*> FullGCAssetClasses;
 	if ( FullGCAssetClassNames.Num() )
diff --git a/Engine/Source/Programs/CrashReporter/CrashReportWebSite/Models/CrashRepository.cs b/Engine/Source/Programs/CrashReporter/CrashReportWebSite/Models/CrashRepository.cs
--- a/Engine/Source/Programs/CrashReporter/CrashReportWebSite/Models/CrashRepository.cs
+++ b/Engine/Source/Programs/CrashReporter/CrashReportWebSite/Models/CrashRepository.cs
@@ -88,7 +88,8 @@ namespace Tools.CrashReporter.CrashReportWebSite.Models
 				if (Now - LastBranchesDate > TimeSpan.FromHours( 1 ))
 				{
 					var BranchList = Context.Crashes
-					.Where( n => n.TimeOfCrash > DateTime.Now.AddMonths( -3 ) )
+					.Where( n => n.TimeOfCrash > DateTime.Now.AddDays( -14 ) )
+                    .Where(c => c.CrashType != 3) // Ignore ensures
 					// Depot - //depot/UE4* || Stream //UE4, //Something etc.
 					.Where( n => n.Branch.StartsWith( "UE4" ) || n.Branch.StartsWith( "//" ) )
 					.Select( n => n.Branch )
@@ -142,7 +143,8 @@ namespace Tools.CrashReporter.CrashReportWebSite.Models
 				if (Now - LastVersionDate > TimeSpan.FromHours( 1 ))
 				{
 					var BuildVersions = Context.Crashes
-					.Where( c => c.TimeOfCrash > DateTime.Now.AddMonths( -3 ) )
+					.Where( c => c.TimeOfCrash > DateTime.Now.AddDays( -14 ) )
+                    .Where(c => c.CrashType != 3 ) // Ignore ensures
 					.Select( c => c.BuildVersion )
 					.Distinct()
 					.ToList();
diff --git a/Engine/Source/Programs/UnrealBuildTool/Configuration/UEBuildModule.cs b/Engine/Source/Programs/UnrealBuildTool/Configuration/UEBuildModule.cs
--- a/Engine/Source/Programs/UnrealBuildTool/Configuration/UEBuildModule.cs
+++ b/Engine/Source/Programs/UnrealBuildTool/Configuration/UEBuildModule.cs
@@ -37,6 +37,7 @@ namespace UnrealBuildTool
 			{
 				case ModuleHostType.Runtime:
 				case ModuleHostType.RuntimeNoCommandlet:
+				case ModuleHostType.RuntimeNoProgram:
                 case ModuleHostType.ServerOnly:
 					return UEBuildModuleType.Runtime;
 				case ModuleHostType.Developer:
@@ -99,6 +100,7 @@ namespace UnrealBuildTool
 					return UEBuildModuleType.Program;
 				case ModuleHostType.Runtime:
 				case ModuleHostType.RuntimeNoCommandlet:
+				case ModuleHostType.RuntimeNoProgram:
 					return UEBuildModuleType.Runtime;
 				default:
 					throw new BuildException("Unhandled module type {0}", Module.Type.ToString());
diff --git a/Engine/Source/Programs/UnrealBuildTool/Configuration/UEBuildTarget.cs b/Engine/Source/Programs/UnrealBuildTool/Configuration/UEBuildTarget.cs
--- a/Engine/Source/Programs/UnrealBuildTool/Configuration/UEBuildTarget.cs
+++ b/Engine/Source/Programs/UnrealBuildTool/Configuration/UEBuildTarget.cs
@@ -884,7 +884,12 @@ namespace UnrealBuildTool
 		/// Additional plugin filenames which are foreign to this target
 		/// </summary>
 		[NonSerialized]
-		public List<FileReference> ForeignPlugins;
+		public List<PluginInfo> UnrealHeaderToolPlugins;
+
+		/// <summary>
+		/// Additional plugin filenames to include when building UnrealHeaderTool for the current target
+		/// </summary>
+		public List<FileReference> ForeignPlugins = new List<FileReference>();
 
 		/// <summary>
 		/// All application binaries; may include binaries not built by this target.
@@ -3204,6 +3209,7 @@ namespace UnrealBuildTool
 
 			// Build a list of enabled plugins
 			EnabledPlugins = new List<PluginInfo>();
+			UnrealHeaderToolPlugins = new List<PluginInfo>();
 
 			// If we're compiling against the engine, add the plugins enabled for this target
 			if (UEBuildConfiguration.bCompileAgainstEngine)
@@ -3213,7 +3219,11 @@ namespace UnrealBuildTool
 				{
 					if (UProjectInfo.IsPluginEnabledForProject(ValidPlugin, Project, Platform))
 					{
-						EnabledPlugins.Add(ValidPlugin);
+						if (ValidPlugin.Descriptor.bCanBeUsedWithUnrealHeaderTool)
+						{
+							UnrealHeaderToolPlugins.Add(ValidPlugin);							
+						}
+						EnabledPlugins.Add(ValidPlugin);						
 					}
 				}
 			}
diff --git a/Engine/Source/Programs/UnrealBuildTool/System/ExternalExecution.cs b/Engine/Source/Programs/UnrealBuildTool/System/ExternalExecution.cs
--- a/Engine/Source/Programs/UnrealBuildTool/System/ExternalExecution.cs
+++ b/Engine/Source/Programs/UnrealBuildTool/System/ExternalExecution.cs
@@ -742,6 +742,15 @@ namespace UnrealBuildTool
 							UBTArguments.Append(" -2013");
 						}
 
+						// Add UHT plugins to UBT command line as external plugins
+						if (Target.ProjectDescriptor != null && Target.UnrealHeaderToolPlugins != null)
+						{
+							foreach (PluginInfo Plugin in Target.UnrealHeaderToolPlugins)
+							{
+								UBTArguments.Append(" -PLUGIN \"" + Plugin.File + "\"");
+							}
+						}
+
 						if (RunExternalExecutable(UnrealBuildTool.GetUBTPath(), UBTArguments.ToString()) != 0)
 						{
 							return false;
diff --git a/Engine/Source/Programs/UnrealBuildTool/System/ModuleDescriptor.cs b/Engine/Source/Programs/UnrealBuildTool/System/ModuleDescriptor.cs
--- a/Engine/Source/Programs/UnrealBuildTool/System/ModuleDescriptor.cs
+++ b/Engine/Source/Programs/UnrealBuildTool/System/ModuleDescriptor.cs
@@ -13,6 +13,7 @@ namespace UnrealBuildTool
 		Default,
 		Runtime,
 		RuntimeNoCommandlet,
+		RuntimeNoProgram,
 		Developer,
 		Editor,
 		EditorNoCommandlet,
@@ -184,6 +185,8 @@ namespace UnrealBuildTool
 				case ModuleHostType.Runtime:
 				case ModuleHostType.RuntimeNoCommandlet:
 					return true;
+				case ModuleHostType.RuntimeNoProgram:
+					return TargetType != TargetRules.TargetType.Program;
 				case ModuleHostType.Developer:
 					return bBuildDeveloperTools;
 				case ModuleHostType.Editor:
diff --git a/Engine/Source/Programs/UnrealBuildTool/System/PluginDescriptor.cs b/Engine/Source/Programs/UnrealBuildTool/System/PluginDescriptor.cs
--- a/Engine/Source/Programs/UnrealBuildTool/System/PluginDescriptor.cs
+++ b/Engine/Source/Programs/UnrealBuildTool/System/PluginDescriptor.cs
@@ -70,6 +70,9 @@ namespace UnrealBuildTool
 		// Marks the plugin as beta in the UI
 		public bool bIsBetaVersion;
 
+		// Whether this plugin can be used by UnrealHeaderTool
+		public bool bCanBeUsedWithUnrealHeaderTool;
+
 		// Set for plugins which are installed
 		public bool bInstalled;
 
@@ -143,6 +146,7 @@ namespace UnrealBuildTool
 				RawObject.TryGetBoolField("CanContainContent", out Descriptor.bCanContainContent);
 				RawObject.TryGetBoolField("IsBetaVersion", out Descriptor.bIsBetaVersion);
 				RawObject.TryGetBoolField("Installed", out Descriptor.bInstalled);
+				RawObject.TryGetBoolField("CanBeUsedWithUnrealHeaderTool", out Descriptor.bCanBeUsedWithUnrealHeaderTool);
 
 				return Descriptor;
 			}
diff --git a/Engine/Source/Programs/UnrealBuildTool/System/UnrealBuildTool.cs b/Engine/Source/Programs/UnrealBuildTool/System/UnrealBuildTool.cs
--- a/Engine/Source/Programs/UnrealBuildTool/System/UnrealBuildTool.cs
+++ b/Engine/Source/Programs/UnrealBuildTool/System/UnrealBuildTool.cs
@@ -1739,6 +1739,19 @@ namespace UnrealBuildTool
 							ReasonNotLoaded = "modules to compile have changed";
 						}
 
+						// Invalid makefile if foreign plugins have changed
+						if (UBTMakefile != null && !TargetDescs.SelectMany(x => x.ForeignPlugins)
+								.Select(x => new Tuple<string, bool>(x.FullName.ToLower(), string.IsNullOrWhiteSpace(x.FullName)))
+								.SequenceEqual(
+									UBTMakefile.Targets.SelectMany(x => x.ForeignPlugins)
+									.Select(x => new Tuple<string, bool>(x.FullName.ToLower(), string.IsNullOrWhiteSpace(x.FullName)))
+								)
+							)
+						{
+							UBTMakefile = null;
+							ReasonNotLoaded = "foreign plugins have changed";
+						}
+
 						if (UBTMakefile == null)
 						{
 							// If the Makefile couldn't be loaded, then we're not going to be able to continue in "assembler only" mode.  We'll do both
@@ -2265,7 +2278,7 @@ namespace UnrealBuildTool
 		[Serializable]
 		class UBTMakefile : ISerializable
 		{
-			public const int CurrentVersion = 3;
+			public const int CurrentVersion = 4;
 
 			public int Version;
 
diff --git a/Engine/Source/Programs/UnrealHeaderTool/Private/UnrealHeaderToolMain.cpp b/Engine/Source/Programs/UnrealHeaderTool/Private/UnrealHeaderToolMain.cpp
--- a/Engine/Source/Programs/UnrealHeaderTool/Private/UnrealHeaderToolMain.cpp
+++ b/Engine/Source/Programs/UnrealHeaderTool/Private/UnrealHeaderToolMain.cpp
@@ -48,8 +48,7 @@ INT32_MAIN_INT32_ARGC_TCHAR_ARGV()
 	{
 		const TCHAR* CmdLinePtr = *ShortCmdLine;
 
-		// Parse the game name.  UHT doesn't care about this directly, but the engine init code will be
-		// expecting it when it parses the command-line itself.
+		// Parse the game name or project filename.  UHT reads the list of plugins from there in case one of the plugins is UHT plugin.
 		FString GameName = FParse::Token(CmdLinePtr, false);
 
 		// This parameter is the absolute path to the file which contains information about the modules
diff --git a/Engine/Source/Programs/UnrealHeaderTool/UnrealHeaderTool.Target.cs b/Engine/Source/Programs/UnrealHeaderTool/UnrealHeaderTool.Target.cs
--- a/Engine/Source/Programs/UnrealHeaderTool/UnrealHeaderTool.Target.cs
+++ b/Engine/Source/Programs/UnrealHeaderTool/UnrealHeaderTool.Target.cs
@@ -8,7 +8,6 @@ public class UnrealHeaderToolTarget : TargetRules
 	public UnrealHeaderToolTarget(TargetInfo Target)
 	{
 		Type = TargetType.Program;
-		AdditionalPlugins.Add("ScriptGeneratorPlugin");
 	}
 
 	//
diff --git a/Engine/Source/Runtime/AIModule/Private/AISystem.cpp b/Engine/Source/Runtime/AIModule/Private/AISystem.cpp
--- a/Engine/Source/Runtime/AIModule/Private/AISystem.cpp
+++ b/Engine/Source/Runtime/AIModule/Private/AISystem.cpp
@@ -128,7 +128,7 @@ void UAISystem::AILoggingVerbose()
 
 void UAISystem::RunEQS(const FString& QueryName, UObject* Target)
 {
-#if !(UE_BUILD_SHIPPING || UE_BUILD_TEST)
+#if !UE_BUILD_SHIPPING
 	UWorld* OuterWorld = GetOuterWorld();
 	if (OuterWorld == NULL || OuterWorld->GetGameInstance() == NULL)
 	{
@@ -155,7 +155,7 @@ void UAISystem::RunEQS(const FString& QueryName, UObject* Target)
 	{
 		MyPC->ClientMessage(TEXT("No debugging target"));
 	}
-#endif
+#endif // !UE_BUILD_SHIPPING
 }
 
 UAISystem::FBlackboardDataToComponentsIterator::FBlackboardDataToComponentsIterator(FBlackboardDataToComponentsMap& InBlackboardDataToComponentsMap, UBlackboardData* BlackboardAsset)
diff --git a/Engine/Source/Runtime/AIModule/Private/AISystemExec.cpp b/Engine/Source/Runtime/AIModule/Private/AISystemExec.cpp
--- a/Engine/Source/Runtime/AIModule/Private/AISystemExec.cpp
+++ b/Engine/Source/Runtime/AIModule/Private/AISystemExec.cpp
@@ -1,7 +1,7 @@
 // Copyright 1998-2015 Epic Games, Inc. All Rights Reserved.
 
 #include "AIModulePrivate.h"
-#if !(UE_BUILD_SHIPPING || UE_BUILD_TEST)
+#if !UE_BUILD_SHIPPING
 #include "Misc/CoreMisc.h"
 #include "BehaviorTree/BehaviorTreeManager.h"
 
@@ -53,4 +53,4 @@ bool FAISystemExec::Exec(UWorld* Inworld, const TCHAR* Cmd, FOutputDevice& Ar)
 	return bHandled;
 }
 
-#endif // !(UE_BUILD_SHIPPING || UE_BUILD_TEST)
\ No newline at end of file
+#endif // !UE_BUILD_SHIPPING
\ No newline at end of file
diff --git a/Engine/Source/Runtime/Core/Private/GenericPlatform/GenericPlatformMemory.cpp b/Engine/Source/Runtime/Core/Private/GenericPlatform/GenericPlatformMemory.cpp
--- a/Engine/Source/Runtime/Core/Private/GenericPlatform/GenericPlatformMemory.cpp
+++ b/Engine/Source/Runtime/Core/Private/GenericPlatform/GenericPlatformMemory.cpp
@@ -181,6 +181,12 @@ uint32 FGenericPlatformMemory::GetPhysicalGBRam()
 	return FPlatformMemory::GetConstants().TotalPhysicalGB;
 }
 
+bool FGenericPlatformMemory::PageProtect( void* const Ptr, const SIZE_T Size, const uint32 ProtectMode )
+{
+	UE_LOG(LogMemory, Error, TEXT("FGenericPlatformMemory::PageProtect not implemented on this platform"));
+	return false;
+}
+
 void* FGenericPlatformMemory::BinnedAllocFromOS( SIZE_T Size )
 {
 	UE_LOG(LogMemory, Error, TEXT("FGenericPlatformMemory::BinnedAllocFromOS not implemented on this platform"));
diff --git a/Engine/Source/Runtime/Core/Private/GenericPlatform/GenericPlatformStackWalk.cpp b/Engine/Source/Runtime/Core/Private/GenericPlatform/GenericPlatformStackWalk.cpp
--- a/Engine/Source/Runtime/Core/Private/GenericPlatform/GenericPlatformStackWalk.cpp
+++ b/Engine/Source/Runtime/Core/Private/GenericPlatform/GenericPlatformStackWalk.cpp
@@ -14,6 +14,22 @@ FProgramCounterSymbolInfo::FProgramCounterSymbolInfo() :
 	FCStringAnsi::Strncpy( Filename, "", MAX_NAME_LENGHT );
 }
 
+/** Settings for stack walking */
+static bool GWantsDetailedCallstacksInNonMonolithicBuilds = true;
+
+void FGenericPlatformStackWalk::Init()
+{
+	// This needs to be called once configs are initialized
+	check(GConfig);
+
+	GConfig->GetBool(TEXT("Core.System"), TEXT("DetailedCallstacksInNonMonolithicBuilds"), GWantsDetailedCallstacksInNonMonolithicBuilds, GEngineIni);
+}
+
+bool FGenericPlatformStackWalk::WantsDetailedCallstacksInNonMonolithicBuilds()
+{
+	return GWantsDetailedCallstacksInNonMonolithicBuilds;
+}
+
 bool FGenericPlatformStackWalk::ProgramCounterToHumanReadableString( int32 CurrentCallDepth, uint64 ProgramCounter, ANSICHAR* HumanReadableString, SIZE_T HumanReadableStringSize, FGenericCrashContext* Context )
 {
 	if (HumanReadableString && HumanReadableStringSize > 0)
diff --git a/Engine/Source/Runtime/Core/Private/HAL/Allocators/AnsiAllocator.h b/Engine/Source/Runtime/Core/Private/HAL/Allocators/AnsiAllocator.h
new file mode 100644
--- /dev/null
+++ b/Engine/Source/Runtime/Core/Private/HAL/Allocators/AnsiAllocator.h
@@ -0,0 +1,114 @@
+// Copyright 1998-2015 Epic Games, Inc. All Rights Reserved.
+
+/*=====================================================================================================
+	AnsiAllocator.h: helper allocator that allocates directly from standard library allocation functions
+=======================================================================================================*/
+
+#pragma once
+
+
+/** Allocator that allocates memory using standard library functions. */
+class CORE_API FAnsiAllocator
+{
+public:
+
+	enum { NeedsElementType = false };
+	enum { RequireRangeCheck = true };
+
+	typedef FAnsiAllocator ElementAllocator;
+	typedef FAnsiAllocator BitArrayAllocator;
+
+	class CORE_API ForAnyElementType
+	{
+	public:
+		/** Default constructor. */
+		ForAnyElementType()
+			: Data(nullptr)
+		{}
+
+		/**
+		* Moves the state of another allocator into this one.
+		* Assumes that the allocator is currently empty, i.e. memory may be allocated but any existing elements have already been destructed (if necessary).
+		* @param Other - The allocator to move the state from.  This allocator should be left in a valid empty state.
+		*/
+		FORCEINLINE void MoveToEmpty(ForAnyElementType& Other)
+		{
+			check(this != &Other);
+
+			if (Data)
+			{
+				::free(Data);
+			}
+
+			Data = Other.Data;
+			Other.Data = nullptr;
+		}
+
+		/** Destructor. */
+		FORCEINLINE ~ForAnyElementType()
+		{
+			if (Data)
+			{
+				::free(Data);
+			}
+		}
+
+		// FContainerAllocatorInterface
+		FORCEINLINE FScriptContainerElement* GetAllocation() const
+		{
+			return Data;
+		}
+		void ResizeAllocation(int32 PreviousNumElements, int32 NumElements, SIZE_T NumBytesPerElement)
+		{
+			// Avoid calling FMemory::Realloc( nullptr, 0 ) as ANSI C mandates returning a valid pointer which is not what we want.
+			if (Data || NumElements)
+			{
+				//checkSlow(((uint64)NumElements*(uint64)ElementTypeInfo.GetSize() < (uint64)INT_MAX));
+				Data = (FScriptContainerElement*)::realloc(Data, NumElements*NumBytesPerElement);
+			}
+		}
+		int32 CalculateSlack(int32 NumElements, int32 NumAllocatedElements, SIZE_T NumBytesPerElement) const
+		{
+			return DefaultCalculateSlack(NumElements, NumAllocatedElements, NumBytesPerElement);
+		}
+
+		SIZE_T GetAllocatedSize(int32 NumAllocatedElements, SIZE_T NumBytesPerElement) const
+		{
+			return NumAllocatedElements * NumBytesPerElement;
+		}
+
+	private:
+		ForAnyElementType(const ForAnyElementType&);
+		ForAnyElementType& operator=(const ForAnyElementType&);
+
+		/** A pointer to the container's elements. */
+		FScriptContainerElement* Data;
+	};
+
+	template<typename ElementType>
+	class ForElementType : public ForAnyElementType
+	{
+	public:
+
+		/** Default constructor. */
+		ForElementType()
+		{}
+
+		FORCEINLINE ElementType* GetAllocation() const
+		{
+			return (ElementType*)ForAnyElementType::GetAllocation();
+		}
+	};
+};
+
+template <>
+struct TAllocatorTraits<FAnsiAllocator> : TAllocatorTraitsBase<FAnsiAllocator>
+{
+	enum { SupportsMove = true };
+	enum { IsZeroConstruct = true };
+};
+
+/** ANSI allocator that can be used with a TSet. */
+class FAnsiSetAllocator : public TSetAllocator<FAnsiAllocator, TInlineAllocator<1, FAnsiAllocator>>
+{	
+};
diff --git a/Engine/Source/Runtime/Core/Private/HAL/ConsoleManager.cpp b/Engine/Source/Runtime/Core/Private/HAL/ConsoleManager.cpp
--- a/Engine/Source/Runtime/Core/Private/HAL/ConsoleManager.cpp
+++ b/Engine/Source/Runtime/Core/Private/HAL/ConsoleManager.cpp
@@ -1497,11 +1497,15 @@ void CreateConsoleVariables()
 	IConsoleManager::Get().RegisterConsoleCommand(TEXT("Vis"),	TEXT("short version of visualizetexture"), ECVF_Cheat);
 	IConsoleManager::Get().RegisterConsoleCommand(TEXT("VisRT"),	TEXT("GUI for visualizetexture"), ECVF_Cheat);
 	IConsoleManager::Get().RegisterConsoleCommand(TEXT("HighResShot"),	TEXT("High resolution screenshots [Magnification = 2..]"), ECVF_Cheat);
-	IConsoleManager::Get().RegisterConsoleCommand(TEXT("DumpConsoleCommands"),	TEXT("Dumps all console vaiables and commands and all exec that can be discovered to the log/console"), ECVF_Default);
 	IConsoleManager::Get().RegisterConsoleCommand(TEXT("DumpUnbuiltLightInteractions"),	TEXT("Logs all lights and primitives that have an unbuilt interaction."), ECVF_Cheat);
 
 #endif // !(UE_BUILD_SHIPPING || UE_BUILD_TEST)
 
+
+#if	!UE_BUILD_SHIPPING
+	IConsoleManager::Get().RegisterConsoleCommand( TEXT( "DumpConsoleCommands" ), TEXT( "Dumps all console vaiables and commands and all exec that can be discovered to the log/console" ), ECVF_Default );
+#endif // !UE_BUILD_SHIPPING
+
 	// testing code
 	{
 		FConsoleManager& ConsoleManager = (FConsoleManager&)IConsoleManager::Get();
diff --git a/Engine/Source/Runtime/Core/Private/HAL/FileManagerGeneric.cpp b/Engine/Source/Runtime/Core/Private/HAL/FileManagerGeneric.cpp
--- a/Engine/Source/Runtime/Core/Private/HAL/FileManagerGeneric.cpp
+++ b/Engine/Source/Runtime/Core/Private/HAL/FileManagerGeneric.cpp
@@ -25,7 +25,7 @@ DEFINE_LOG_CATEGORY_STATIC( LogFileManager, Log, All );
 
 void FFileManagerGeneric::ProcessCommandLineOptions() 
 {
-#if !( UE_BUILD_SHIPPING || UE_BUILD_TEST )
+#if !UE_BUILD_SHIPPING
 	if( FParse::Param( FCommandLine::Get(),TEXT( "CLEANSCREENSHOTS" ) ) )
 	{
 		DeleteDirectory( *FPaths::ScreenShotDir(), false, true );
diff --git a/Engine/Source/Runtime/Core/Private/HAL/MallocBinned2.cpp b/Engine/Source/Runtime/Core/Private/HAL/MallocBinned2.cpp
--- a/Engine/Source/Runtime/Core/Private/HAL/MallocBinned2.cpp
+++ b/Engine/Source/Runtime/Core/Private/HAL/MallocBinned2.cpp
@@ -9,6 +9,7 @@
 #include "MallocBinned2.h"
 #include "MemoryMisc.h"
 #include "HAL/PlatformAtomics.h"
+#include "MallocVerify.h"
 
 //when modifying the global allocator stats, if we are using COARSE locks, then all callsites for stat modification are covered by the allocator-wide access guard. Thus the stats can be modified directly.
 //If we are using FINE locks, then we must modify the stats through atomics as the locks are either not actually covering the stat callsites, or are locking specific table locks which is not sufficient for stats.
@@ -55,6 +56,7 @@ DEFINE_STAT(STAT_Binned2_CurrentAllocs);
 DEFINE_STAT(STAT_Binned2_TotalAllocs);
 DEFINE_STAT(STAT_Binned2_SlackCurrent);
 
+
 // Block sizes are based around getting the maximum amount of allocations per pool, with as little alignment waste as possible.
 // Block sizes should be close to even divisors of the POOL_SIZE, and well distributed. They must be 16-byte aligned as well.
 const uint32 GMallocBinned2BlockSizes[] =
@@ -569,6 +571,9 @@ void* FMallocBinned2::Malloc(SIZE_T Size, uint32 Alignment)
 		}
 
 		FFreeMem* Result = Private::AllocateBlockFromPool(*this, *Table, Pool, Alignment);
+#if MALLOC_VERIFY
+		FMallocVerify::Get().Malloc(Result);
+#endif
 		return Result;
 	}
 
@@ -590,6 +595,9 @@ void* FMallocBinned2::Malloc(SIZE_T Size, uint32 Alignment)
 		}
 
 		FFreeMem* Result = Private::AllocateBlockFromPool(*this, *Table, Pool, Alignment);
+#if MALLOC_VERIFY
+		FMallocVerify::Get().Malloc(Result);
+#endif
 		return Result;
 	}
 
@@ -627,11 +635,14 @@ void* FMallocBinned2::Malloc(SIZE_T Size, uint32 Alignment)
 	BINNED2_PEAK_STATCOUNTER(Stats.UsedPeak,  BINNED2_ADD_STATCOUNTER(Stats.UsedCurrent,  Size));
 	BINNED2_PEAK_STATCOUNTER(Stats.WastePeak, BINNED2_ADD_STATCOUNTER(Stats.WasteCurrent, (int64)(AlignedSize - Size)));
 
+#if MALLOC_VERIFY
+	FMallocVerify::Get().Malloc(AlignedResult);
+#endif
 	return AlignedResult;
 }
 
 void* FMallocBinned2::Realloc(void* Ptr, SIZE_T NewSize, uint32 Alignment)
-{
+{	
 	// Handle DEFAULT_ALIGNMENT for binned allocator.
 	Alignment = FMath::Max<uint32>(Alignment, Private::DEFAULT_BINNED_ALLOCATOR_ALIGNMENT);
 
@@ -641,6 +652,8 @@ void* FMallocBinned2::Realloc(void* Ptr, SIZE_T NewSize, uint32 Alignment)
 		return Result;
 	}
 
+	//ValidateHeap();
+
 	if (NewSize == 0)
 	{
 		FMallocBinned2::Free(Ptr);
@@ -664,7 +677,6 @@ void* FMallocBinned2::Realloc(void* Ptr, SIZE_T NewSize, uint32 Alignment)
 			void* Result = FMallocBinned2::Malloc(NewSizeUnmodified, Alignment);
 			FMemory::Memcpy(Result, Ptr, FMath::Min<SIZE_T>(NewSizeUnmodified, MemSizeToPoolTable[Pool->TableIndex]->BlockSize - (Alignment - SpareBytesCount)));
 			FMallocBinned2::Free(Ptr);
-
 			return Result;
 		}
 
@@ -672,6 +684,9 @@ void* FMallocBinned2::Realloc(void* Ptr, SIZE_T NewSize, uint32 Alignment)
 		{
 			void* Result = Align(Ptr, Alignment);
 			FMemory::Memmove(Result, Ptr, NewSize);
+#if MALLOC_VERIFY
+			FMallocVerify::Get().Realloc(Ptr, Result);
+#endif
 			return Result;
 		}
 
@@ -686,7 +701,6 @@ void* FMallocBinned2::Realloc(void* Ptr, SIZE_T NewSize, uint32 Alignment)
 		void* Result = FMallocBinned2::Malloc(NewSizeUnmodified, Alignment);
 		FMemory::Memcpy(Result, Ptr, FMath::Min<SIZE_T>(NewSizeUnmodified, Pool->AllocSize));
 		FMallocBinned2::Free(Ptr);
-
 		return Result;
 	}
 
@@ -700,6 +714,9 @@ void* FMallocBinned2::Realloc(void* Ptr, SIZE_T NewSize, uint32 Alignment)
 
 	Pool->SetAllocationSizes(NewSizeUnmodified, PoolOsBytes, BinnedOSTableIndex, BinnedOSTableIndex);
 
+#if MALLOC_VERIFY
+	FMallocVerify::Get().Realloc(Ptr, Ptr);
+#endif
 	return Ptr;
 }
 
@@ -710,6 +727,10 @@ void FMallocBinned2::Free( void* Ptr )
 		return;
 	}
 
+#if MALLOC_VERIFY
+	FMallocVerify::Get().Free(Ptr);
+#endif
+
 	BINNED2_DECREMENT_STATCOUNTER(Stats.CurrentAllocs);
 
 	void* BasePtr;
diff --git a/Engine/Source/Runtime/Core/Private/HAL/MallocStomp.cpp b/Engine/Source/Runtime/Core/Private/HAL/MallocStomp.cpp
new file mode 100644
--- /dev/null
+++ b/Engine/Source/Runtime/Core/Private/HAL/MallocStomp.cpp
@@ -0,0 +1,141 @@
+// Copyright 1998-2015 Epic Games, Inc. All Rights Reserved.
+
+#include "CorePrivatePCH.h"
+
+#if USE_MALLOC_STOMP
+
+#include "MallocStomp.h"
+
+static void MallocStompOverrunTest()
+{
+	const uint32 ArraySize = 4;
+	uint8* Pointer = new uint8[ArraySize];
+	// Overrun.
+	Pointer[ArraySize+1] = 0;
+}
+
+FAutoConsoleCommand MallocStompTestCommand
+(
+	TEXT( "MallocStomp.OverrunTest" ),
+	TEXT( "Overrun test for the FMallocStomp" ),
+	FConsoleCommandDelegate::CreateStatic( &MallocStompOverrunTest )
+);
+
+
+void* FMallocStomp::Malloc(SIZE_T Size, uint32 Alignment)
+{
+	if(Size == 0U)
+	{
+		return nullptr;
+	}
+
+	const SIZE_T AlignedSize = (Alignment > 0U) ? ((Size + Alignment - 1U) & -static_cast<int32>(Alignment)) : Size;
+	const SIZE_T AllocFullPageSize = AlignedSize + sizeof(FAllocationData) + (PageSize - 1) & ~(PageSize - 1U);
+
+#if PLATFORM_LINUX || PLATFORM_MAC
+	// Note: can't implement BinnedAllocFromOS as a mmap call. See Free() for the reason.
+	void *FullAllocationPointer = mmap(nullptr, AllocFullPageSize + PageSize, PROT_READ | PROT_WRITE, MAP_PRIVATE | MAP_ANON, -1, 0);
+#else
+	void *FullAllocationPointer = FPlatformMemory::BinnedAllocFromOS(AllocFullPageSize + PageSize);
+#endif // PLATFORM_LINUX || PLATFORM_MAC
+
+	void *ReturnedPointer = nullptr;
+	static const SIZE_T AllocationDataSize = sizeof(FAllocationData);
+
+	if(bUseUnderrunMode)
+	{
+		const SIZE_T AlignedAllocationData = (Alignment > 0U) ? ((AllocationDataSize + Alignment - 1U) & -static_cast<int32>(Alignment)) : AllocationDataSize;
+		ReturnedPointer = reinterpret_cast<void*>(reinterpret_cast<uint8*>(FullAllocationPointer) + PageSize + AlignedAllocationData);
+
+		FAllocationData *AllocDataPtr = reinterpret_cast<FAllocationData*>(reinterpret_cast<uint8*>(FullAllocationPointer) + PageSize);
+		FAllocationData AllocData = { FullAllocationPointer, AllocFullPageSize + PageSize, AlignedSize, SentinelExpectedValue };
+		*AllocDataPtr = AllocData;
+
+		// Page protect the first page, this will cause the exception in case the is an underrun.
+		FPlatformMemory::PageProtect(FullAllocationPointer, PageSize, NoAccessProtectMode);
+	}
+	else
+	{
+		ReturnedPointer = reinterpret_cast<void*>(reinterpret_cast<uint8*>(FullAllocationPointer) + AllocFullPageSize - AlignedSize);
+
+		FAllocationData *AllocDataPtr = reinterpret_cast<FAllocationData*>(reinterpret_cast<uint8*>(ReturnedPointer) - AllocationDataSize);
+		FAllocationData AllocData = { FullAllocationPointer, AllocFullPageSize + PageSize, AlignedSize, SentinelExpectedValue };
+		*AllocDataPtr = AllocData;
+
+		// Page protect the last page, this will cause the exception in case the is an overrun.
+		FPlatformMemory::PageProtect(reinterpret_cast<void*>(reinterpret_cast<uint8*>(FullAllocationPointer) + AllocFullPageSize), PageSize, NoAccessProtectMode);
+	}
+
+	return ReturnedPointer;
+}
+
+void* FMallocStomp::Realloc(void* InPtr, SIZE_T NewSize, uint32 Alignment)
+{
+	if(NewSize == 0U)
+	{
+		Free(InPtr);
+		return nullptr;
+	}
+
+	void *ReturnPtr = nullptr;
+	if(InPtr != nullptr)
+	{
+		ReturnPtr = Malloc(NewSize, Alignment);
+
+		FAllocationData *AllocDataPtr = reinterpret_cast<FAllocationData*>(reinterpret_cast<uint8*>(InPtr) - sizeof(FAllocationData));
+		FMemory::Memcpy(ReturnPtr, InPtr, FMath::Min(AllocDataPtr->Size, NewSize));
+		Free(InPtr);
+	}
+	else
+	{
+		ReturnPtr = Malloc(NewSize, Alignment);
+	}
+
+	return ReturnPtr;
+}
+
+void FMallocStomp::Free(void* InPtr)
+{
+	if(InPtr == nullptr)
+	{
+		return;
+	}
+
+	FAllocationData *AllocDataPtr = reinterpret_cast<FAllocationData*>(InPtr);
+	AllocDataPtr--;
+
+	// Check that our sentinel is intact.
+	if(AllocDataPtr->Sentinel != SentinelExpectedValue)
+	{
+		// There was a memory underrun related to this allocation.
+		FPlatformMisc::DebugBreak();
+	}
+
+#if PLATFORM_LINUX || PLATFORM_MAC
+	// Note: Can't wrap munmap inside BinnedFreeToOS() because the code doesn't
+	// expect the size of the allocation to be freed to be available, nor the 
+	// pointer be aligned with the page size. We can guarantee that here so that's
+	// why we can do it.
+	munmap(AllocDataPtr->FullAllocationPointer, AllocDataPtr->FullSize);
+#else
+	FPlatformMemory::BinnedFreeToOS(AllocDataPtr->FullAllocationPointer);
+#endif // PLATFORM_LINUX || PLATFORM_MAC
+}
+
+bool FMallocStomp::GetAllocationSize(void *Original, SIZE_T &SizeOut) 
+{
+	if(Original == nullptr)
+	{
+		SizeOut = 0U;
+	}
+	else
+	{
+		FAllocationData *AllocDataPtr = reinterpret_cast<FAllocationData*>(Original);
+		AllocDataPtr--;
+		SizeOut = AllocDataPtr->FullSize;
+	}
+
+	return true;
+}
+
+#endif // USE_MALLOC_STOMP
\ No newline at end of file
diff --git a/Engine/Source/Runtime/Core/Private/HAL/MallocTBB.cpp b/Engine/Source/Runtime/Core/Private/HAL/MallocTBB.cpp
--- a/Engine/Source/Runtime/Core/Private/HAL/MallocTBB.cpp
+++ b/Engine/Source/Runtime/Core/Private/HAL/MallocTBB.cpp
@@ -6,6 +6,7 @@
 
 #include "CorePrivatePCH.h"
 #include "MallocTBB.h"
+#include "MallocVerify.h"
 
 // Only use for supported platforms
 #if PLATFORM_SUPPORTS_TBB && TBB_ALLOCATOR_ALLOWED
@@ -53,6 +54,9 @@ void* FMallocTBB::Malloc( SIZE_T Size, uint32 Alignment )
 	{
 		FMemory::Memset(NewPtr, DEBUG_FILL_NEW, Size); 
 	}
+#endif
+#if MALLOC_VERIFY
+	FMallocVerify::Get().Malloc(NewPtr);
 #endif
 	MEM_TIME(MemTime += FPlatformTime::Seconds());
 	return NewPtr;
@@ -94,6 +98,9 @@ void* FMallocTBB::Realloc( void* Ptr, SIZE_T NewSize, uint32 Alignment )
 	{
 		OutOfMemory(NewSize, Alignment);
 	}
+#if MALLOC_VERIFY
+	FMallocVerify::Get().Realloc(Ptr, NewPtr);
+#endif
 	MEM_TIME(MemTime += FPlatformTime::Seconds())
 	return NewPtr;
 }
@@ -105,6 +112,9 @@ void FMallocTBB::Free( void* Ptr )
 		return;
 	}
 	MEM_TIME(MemTime -= FPlatformTime::Seconds())
+#if MALLOC_VERIFY
+	FMallocVerify::Get().Free(Ptr);
+#endif
 #if UE_BUILD_DEBUG || UE_BUILD_DEVELOPMENT
 	FMemory::Memset(Ptr, DEBUG_FILL_FREED, scalable_msize(Ptr)); 
 #endif
diff --git a/Engine/Source/Runtime/Core/Private/HAL/MallocVerify.cpp b/Engine/Source/Runtime/Core/Private/HAL/MallocVerify.cpp
new file mode 100644
--- /dev/null
+++ b/Engine/Source/Runtime/Core/Private/HAL/MallocVerify.cpp
@@ -0,0 +1,63 @@
+// Copyright 1998-2015 Epic Games, Inc. All Rights Reserved.
+
+/*=============================================================================
+	MallocVerify.cpp: Helper class to track memory allocations
+=============================================================================*/
+
+#pragma once
+
+#include "CorePrivatePCH.h"
+#include "MemoryMisc.h"
+#include "MallocVerify.h"
+
+#if MALLOC_VERIFY
+
+void FMallocVerify::Malloc(void* Ptr)
+{
+	if (Ptr)
+	{
+		FScopeLock Lock(&AllocatedPointersCritical);
+
+		UE_CLOG(AllocatedPointers.Contains(Ptr), LogMemory, Fatal, TEXT("Malloc allocated pointer that's already been allocated: 0x%016llx"), (int64)(PTRINT)Ptr);
+
+		AllocatedPointers.Add(Ptr);			
+	}
+}
+
+void FMallocVerify::Realloc(void* OldPtr, void* NewPtr)
+{
+	FScopeLock Lock(&AllocatedPointersCritical);
+	if (OldPtr != NewPtr)
+	{			
+		if (OldPtr)
+		{
+			UE_CLOG(!AllocatedPointers.Contains(OldPtr), LogMemory, Fatal, TEXT("Realloc entered with old pointer that hasn't been allocated yet: 0x%016llx"), (int64)(PTRINT)OldPtr);
+
+			AllocatedPointers.Remove(OldPtr);
+		}
+		if (NewPtr)
+		{
+			UE_CLOG(AllocatedPointers.Contains(NewPtr), LogMemory, Fatal, TEXT("Realloc allocated new pointer that has already been allocated: 0x%016llx"), (int64)(PTRINT)NewPtr);
+
+			AllocatedPointers.Add(NewPtr);
+		}
+	}
+	else if (OldPtr)
+	{
+		UE_CLOG(!AllocatedPointers.Contains(OldPtr), LogMemory, Fatal, TEXT("Realloc entered with old pointer that hasn't been allocated yet: 0x%016llx"), (int64)(PTRINT)OldPtr);
+	}
+}
+
+void FMallocVerify::Free(void* Ptr)
+{
+	if (Ptr)
+	{
+		FScopeLock Lock(&AllocatedPointersCritical);
+
+		UE_CLOG(!AllocatedPointers.Contains(Ptr), LogMemory, Fatal, TEXT("Free attempts to free a pointer that hasn't been allocated yet: 0x%016llx"), (int64)(PTRINT)Ptr);
+
+		AllocatedPointers.Remove(Ptr);
+	}
+}
+
+#endif // MALLOC_VERIFY
\ No newline at end of file
diff --git a/Engine/Source/Runtime/Core/Private/HAL/MallocVerify.h b/Engine/Source/Runtime/Core/Private/HAL/MallocVerify.h
new file mode 100644
--- /dev/null
+++ b/Engine/Source/Runtime/Core/Private/HAL/MallocVerify.h
@@ -0,0 +1,43 @@
+// Copyright 1998-2015 Epic Games, Inc. All Rights Reserved.
+
+/*=============================================================================
+	MallocVerify.h: Helper class to track memory allocations
+=============================================================================*/
+
+#pragma once
+
+#include "Allocators/AnsiAllocator.h"
+
+#define MALLOC_VERIFY 0
+
+#if MALLOC_VERIFY
+
+/**
+/* Maintains a list of all pointers to currently allocated memory.
+ */
+class FMallocVerify
+{
+	/** List of all currently allocated pointers */
+	TSet<void*, DefaultKeyFuncs<void*>, FAnsiSetAllocator> AllocatedPointers;
+	FCriticalSection AllocatedPointersCritical;
+
+public:
+
+	/** Handles new allocated pointer */
+	void Malloc(void* Ptr);
+
+	/** Handles reallocation */
+	void Realloc(void* OldPtr, void* NewPtr);
+
+	/** Removes allocated pointer from list */
+	void Free(void* Ptr);
+
+	/** Returns singleton object */
+	static FMallocVerify& Get()
+	{
+		static FMallocVerify Instance;
+		return Instance;
+	}
+};
+
+#endif // MALLOC_VERIFY
diff --git a/Engine/Source/Runtime/Core/Private/Linux/LinuxPlatformMemory.cpp b/Engine/Source/Runtime/Core/Private/Linux/LinuxPlatformMemory.cpp
--- a/Engine/Source/Runtime/Core/Private/Linux/LinuxPlatformMemory.cpp
+++ b/Engine/Source/Runtime/Core/Private/Linux/LinuxPlatformMemory.cpp
@@ -115,6 +115,11 @@ class FMalloc* FLinuxPlatformMemory::BaseAllocator()
 	return Allocator;
 }
 
+bool FLinuxPlatformMemory::PageProtect( void* const Ptr, const SIZE_T Size, const uint32 ProtectMode )
+{
+	return mprotect(Ptr, Size, static_cast<int32>(ProtectMode)) == 0;
+}
+
 void* FLinuxPlatformMemory::BinnedAllocFromOS( SIZE_T Size )
 {	
 	return valloc(Size);	// equivalent to memalign(sysconf(_SC_PAGESIZE),size).
diff --git a/Engine/Source/Runtime/Core/Private/Mac/MacPlatformMemory.cpp b/Engine/Source/Runtime/Core/Private/Mac/MacPlatformMemory.cpp
--- a/Engine/Source/Runtime/Core/Private/Mac/MacPlatformMemory.cpp
+++ b/Engine/Source/Runtime/Core/Private/Mac/MacPlatformMemory.cpp
@@ -8,6 +8,9 @@
 #include "MallocTBB.h"
 #include "MallocAnsi.h"
 #include "MallocBinned2.h"
+#if USE_MALLOC_STOMP
+	#include "MallocStomp.h"
+#endif // USE_MALLOC_STOMP
 
 #include <sys/param.h>
 #include <sys/mount.h>
@@ -47,6 +50,8 @@ FMalloc* FMacPlatformMemory::BaseAllocator()
 	{
 #if FORCE_ANSI_ALLOCATOR || IS_PROGRAM
 		return new FMallocAnsi();
+#elif USE_MALLOC_STOMP
+		return new FMallocStomp();
 #elif (WITH_EDITORONLY_DATA || IS_PROGRAM) && TBB_ALLOCATOR_ALLOWED
 		return new FMallocTBB();
 #else
@@ -126,6 +131,11 @@ const FPlatformMemoryConstants& FMacPlatformMemory::GetConstants()
 	return MemoryConstants;	
 }
 
+bool FMacPlatformMemory::PageProtect( void* const Ptr, const SIZE_T Size, const uint32 ProtectMode )
+{
+	return mprotect(Ptr, Size, static_cast<int32>(ProtectMode)) == 0;
+}
+
 void* FMacPlatformMemory::BinnedAllocFromOS( SIZE_T Size )
 {
 	return valloc(Size);
diff --git a/Engine/Source/Runtime/Core/Private/Misc/EventPool.h b/Engine/Source/Runtime/Core/Private/Misc/EventPool.h
--- a/Engine/Source/Runtime/Core/Private/Misc/EventPool.h
+++ b/Engine/Source/Runtime/Core/Private/Misc/EventPool.h
@@ -77,10 +77,7 @@ public:
 		check(Event);
 		check(Event->IsManualReset() == (PoolType == EEventPoolTypes::ManualReset));
 
-		if (PoolType == EEventPoolTypes::ManualReset)
-		{
-			Event->Reset();
-		}
+		Event->Reset();
 
 		Pool.Push(Event);
 	}
diff --git a/Engine/Source/Runtime/Core/Private/Misc/Parse.cpp b/Engine/Source/Runtime/Core/Private/Misc/Parse.cpp
--- a/Engine/Source/Runtime/Core/Private/Misc/Parse.cpp
+++ b/Engine/Source/Runtime/Core/Private/Misc/Parse.cpp
@@ -3,7 +3,7 @@
 #include "CorePrivatePCH.h"
 #include "LazyPrintf.h"
 
-#if !(UE_BUILD_SHIPPING || UE_BUILD_TEST)
+#if !UE_BUILD_SHIPPING 
 /**
  * Needed for the console command "DumpConsoleCommands"
  * How it works:
@@ -199,7 +199,7 @@ void ConsoleCommandLibrary_DumpLibraryHTML(UWorld* InWorld, FExec& SubSystem, co
 	}
 */
 }
-#endif // !(UE_BUILD_SHIPPING || UE_BUILD_TEST)
+#endif // UE_BUILD_SHIPPING
 
 //
 // Get a string from a text string.
@@ -665,7 +665,7 @@ bool FParse::Value( const TCHAR* Stream, const TCHAR* Match, struct FGuid& Guid
 //
 bool FParse::Command( const TCHAR** Stream, const TCHAR*  Match, bool bParseMightTriggerExecution )
 {
-#if !(UE_BUILD_SHIPPING || UE_BUILD_TEST)
+#if !UE_BUILD_SHIPPING
 	if(GConsoleCommandLibrary)
 	{
 		GConsoleCommandLibrary->OnParseCommand(Match);
@@ -676,7 +676,7 @@ bool FParse::Command( const TCHAR** Stream, const TCHAR*  Match, bool bParseMigh
 			return false;
 		}
 	}
-#endif // !(UE_BUILD_SHIPPING || UE_BUILD_TEST)
+#endif // !UE_BUILD_SHIPPING
 
 	while( (**Stream==' ')||(**Stream==9) )
 		(*Stream)++;
diff --git a/Engine/Source/Runtime/Core/Private/Stats/StatsCommand.cpp b/Engine/Source/Runtime/Core/Private/Stats/StatsCommand.cpp
--- a/Engine/Source/Runtime/Core/Private/Stats/StatsCommand.cpp
+++ b/Engine/Source/Runtime/Core/Private/Stats/StatsCommand.cpp
@@ -1676,38 +1676,8 @@ static void PrintStatsHelpToOutputDevice( FOutputDevice& Ar )
 	Ar.Log( TEXT("stat stopfile - stops tracking all memory operations and writes the results to the file"));
 
 	Ar.Log( TEXT("stat namedmarker #markername# - adds a custom marker to the stats stream"));
-}
-
-// @TODO yrx 2014-12-01 Move to StatsFile.cpp/.h
-static void CommandTestFile()
-{
-	const FString& LastFileSaved = FCommandStatsFile::Get().LastFileSaved;
 
-	FStatsThreadState Loaded( LastFileSaved );
-	if( Loaded.GetLatestValidFrame() < 0 )
-	{
-		UE_LOG( LogStats, Log, TEXT( "Failed to stats file: %s" ), *LastFileSaved );
-		return;
-	}
-	UE_LOG( LogStats, Log, TEXT( "Loaded stats file: %s, %lld frame" ), *LastFileSaved, 1 + Loaded.GetLatestValidFrame() - Loaded.GetOldestValidFrame() );
-	{
-		int64 TestFrame = Loaded.GetOldestValidFrame();
-		UE_LOG( LogStats, Log, TEXT( "**************************** Test Frame %lld" ), TestFrame );
-		DumpHistoryFrame( Loaded, TestFrame );
-	}
-	{
-		int64 TestFrame = (Loaded.GetLatestValidFrame() + Loaded.GetOldestValidFrame()) / 2;
-		if( Loaded.IsFrameValid( TestFrame ) )
-		{
-			UE_LOG( LogStats, Log, TEXT( "**************************** Test Frame %lld" ), TestFrame );
-			DumpHistoryFrame( Loaded, TestFrame );
-		}
-	}
-	{
-		int64 TestFrame = Loaded.GetLatestValidFrame();
-		UE_LOG( LogStats, Log, TEXT( "**************************** Test Frame %lld" ), TestFrame );
-		DumpHistoryFrame( Loaded, TestFrame );
-	}
+	Ar.Log( TEXT( "stat testfile - loads the last saved capture and dumps first, middle and last frame" ) );
 }
 
 #endif
@@ -1866,7 +1836,7 @@ static void StatCmd(FString InCmd, bool bStatCommand)
 		}
 		else if (FParse::Command(&Cmd, TEXT("TESTFILE")))
 		{
-			CommandTestFile();
+			FCommandStatsFile::Get().TestLastSaved();
 		}
 		else if (FParse::Command(&Cmd, TEXT("testdisable")))
 		{
diff --git a/Engine/Source/Runtime/Core/Private/Stats/StatsData.cpp b/Engine/Source/Runtime/Core/Private/Stats/StatsData.cpp
--- a/Engine/Source/Runtime/Core/Private/Stats/StatsData.cpp
+++ b/Engine/Source/Runtime/Core/Private/Stats/StatsData.cpp
@@ -574,16 +574,9 @@ FStatsThreadState::FStatsThreadState(int32 InHistoryFrames)
 	, bFindMemoryExtensiveStats(false)
 	, CurrentGameFrame(1)
 	, CurrentRenderFrame(1)
-	, MaxFrameSeen(0)
-	, MinFrameSeen(-1)
-	, bWasLoaded(false)
 {
 }
 
-//FStatsThreadState::FStatsThreadState(FString const& Filename)
-//void FStatsThreadState::AddMessages(TArray<FStatMessage>& InMessages)
-//@see moved to StatsFile.cpp
-
 FStatsThreadState& FStatsThreadState::GetLocalState()
 {
 	static FStatsThreadState Singleton;
@@ -592,14 +585,6 @@ FStatsThreadState& FStatsThreadState::GetLocalState()
 
 int64 FStatsThreadState::GetOldestValidFrame() const
 {
-	if (bWasLoaded)
-	{
-		if (MaxFrameSeen < 0 || MinFrameSeen < 0)
-		{
-			return -1;
-		}
-		return MinFrameSeen;
-	}
 	int64 Result = -1;
 	for (auto It = GoodFrames.CreateConstIterator(); It; ++It)
 	{
@@ -613,18 +598,6 @@ int64 FStatsThreadState::GetOldestValidFrame() const
 
 int64 FStatsThreadState::GetLatestValidFrame() const
 {
-	if (bWasLoaded)
-	{
-		if (MaxFrameSeen < 0 || MinFrameSeen < 0)
-		{
-			return -1;
-		}
-		if (MaxFrameSeen > MinFrameSeen)
-		{
-			return MaxFrameSeen - 1;
-		}
-		return MaxFrameSeen;
-	}
 	int64 Result = -1;
 	for (auto It = GoodFrames.CreateConstIterator(); It; ++It)
 	{
@@ -751,7 +724,6 @@ void FStatsThreadState::ToggleFindMemoryExtensiveStats()
 
 void FStatsThreadState::ProcessNonFrameStats(FStatMessagesArray& Data, TSet<FName>* NonFrameStatsFound)
 {
-	check(!bWasLoaded);
 	for (int32 Index = 0; Index < Data.Num() ; Index++)
 	{
 		FStatMessage& Item = Data[Index];
@@ -836,8 +808,6 @@ void FStatsThreadState::AddToHistoryAndEmpty(FStatPacketArray& NewData)
 	History.GenerateKeyArray(Frames);
 	Frames.Sort();
 
-	check(!bWasLoaded);
-
 	int64 LatestFinishedFrame = FMath::Min<int64>(CurrentGameFrame, CurrentRenderFrame) - 1;
 
 	for (int32 Index = 0; Index < Frames.Num(); Index++)
diff --git a/Engine/Source/Runtime/Core/Private/Stats/StatsFile.cpp b/Engine/Source/Runtime/Core/Private/Stats/StatsFile.cpp
--- a/Engine/Source/Runtime/Core/Private/Stats/StatsFile.cpp
+++ b/Engine/Source/Runtime/Core/Private/Stats/StatsFile.cpp
@@ -72,59 +72,11 @@ public:
 };
 
 /*-----------------------------------------------------------------------------
-	FStatsThreadState
+	FStatsLoadedState
 -----------------------------------------------------------------------------*/
 
-FStatsThreadState::FStatsThreadState( FString const& Filename )
-	: HistoryFrames( MAX_int32 )
-	, LastFullFrameMetaAndNonFrame( -1 )
-	, LastFullFrameProcessed( -1 )
-	, TotalNumStatMessages( 0 )
-	, MaxNumStatMessages( 0 )
-	, bFindMemoryExtensiveStats( false )
-	, CurrentGameFrame( -1 )
-	, CurrentRenderFrame( -1 )
-	, MaxFrameSeen( -1 )
-	, MinFrameSeen( -1 )
-	, bWasLoaded( true )
-{}
-
-// #YRX_remove 2015-11-17 
-void FStatsThreadState::AddCondensedMessages( TArray<FStatMessage>& CondensedMessages )
+void FStatsLoadedState::ProcessMetaDataAndLeaveDataOnly( TArray<FStatMessage>& CondensedMessages )
 {
-	bWasLoaded = true;
-	TArray<FStatMessage> Messages;
-	for (int32 Index = 0; Index < CondensedMessages.Num(); ++Index)
-	{
-		if (CondensedMessages[Index].NameAndInfo.GetField<EStatOperation>() == EStatOperation::AdvanceFrameEventGameThread)
-		{
-			ProcessMetaDataForLoad( Messages );
-			if (!CondensedStackHistory.Contains( CurrentGameFrame ) && Messages.Num())
-			{
-				TArray<FStatMessage>* Save = new TArray<FStatMessage>();
-				Exchange( *Save, Messages );
-				if (CondensedStackHistory.Num() >= HistoryFrames)
-				{
-					for (auto& It : CondensedStackHistory)
-					{
-						delete It.Value;
-					}
-					CondensedStackHistory.Reset();
-				}
-				CondensedStackHistory.Add( CurrentGameFrame, Save );
-				GoodFrames.Add( CurrentGameFrame );
-			}
-		}
-
-		new (Messages)FStatMessage( CondensedMessages[Index] );
-	}
-	//bWasLoaded = false;
-}
-
-// #YRX_remove 2015-11-17 
-void FStatsThreadState::ProcessMetaDataForLoad( TArray<FStatMessage>& CondensedMessages )
-{
-	check( bWasLoaded );
 	for (int32 Index = 0; Index < CondensedMessages.Num(); Index++)
 	{
 		FStatMessage& Item = CondensedMessages[Index];
@@ -135,42 +87,14 @@ void FStatsThreadState::ProcessMetaDataForLoad( TArray<FStatMessage>& CondensedM
 		}
 		else if (Op == EStatOperation::AdvanceFrameEventGameThread)
 		{
-			check( Item.NameAndInfo.GetField<EStatDataType>() == EStatDataType::ST_int64 );
-			if (Item.GetValue_int64() > 0)
-			{
-				CurrentGameFrame = Item.GetValue_int64();
-				if (CurrentGameFrame > MaxFrameSeen)
-				{
-					MaxFrameSeen = CurrentGameFrame;
-				}
-				if (MinFrameSeen < 0)
-				{
-					MinFrameSeen = CurrentGameFrame;
-				}
-			}
-		}
-		else if (Op == EStatOperation::AdvanceFrameEventRenderThread)
-		{
-			check( Item.NameAndInfo.GetField<EStatDataType>() == EStatDataType::ST_int64 );
-			if (Item.GetValue_int64() > 0)
-			{
-				CurrentRenderFrame = Item.GetValue_int64();
-				if (CurrentRenderFrame > MaxFrameSeen)
-				{
-					MaxFrameSeen = CurrentRenderFrame;
-				}
-				if (MinFrameSeen < 0)
-				{
-					MinFrameSeen = CurrentRenderFrame;
-				}
-			}
+			// The following messages contain only the stats data.
+			CondensedMessages.RemoveAt( 0, Index, true );
+			break;
 		}
 	}
-
-	//UE_LOG( LogStats, Log, TEXT( "Current GameFrame: %lli, RenderFrame: %lli" ), CurrentGameFrame, CurrentRenderFrame );
 }
 
-void FStatsThreadState::AddFrameFromCondensedMessages( TArray<FStatMessage>& CondensedMessages )
+void FStatsLoadedState::AddFrameFromCondensedMessages( TArray<FStatMessage>& CondensedMessages )
 {
 	// @see FStatsThreadState::Condense
 	AdvanceFrameForLoad( CondensedMessages );
@@ -205,7 +129,29 @@ void FStatsThreadState::AddFrameFromCondensedMessages( TArray<FStatMessage>& Con
 	}
 }
 
-void FStatsThreadState::AdvanceFrameForLoad( TArray<FStatMessage>& CondensedMessages )
+int64 FStatsLoadedState::GetOldestValidFrame() const
+{
+	if (MaxFrameSeen < 0 || MinFrameSeen < 0)
+	{
+		return -1;
+	}
+	return MinFrameSeen;
+}
+
+int64 FStatsLoadedState::GetLatestValidFrame() const
+{
+	if (MaxFrameSeen < 0 || MinFrameSeen < 0)
+	{
+		return -1;
+	}
+	if (MaxFrameSeen > MinFrameSeen)
+	{
+		return MaxFrameSeen - 1;
+	}
+	return MaxFrameSeen;
+}
+
+void FStatsLoadedState::AdvanceFrameForLoad( TArray<FStatMessage>& CondensedMessages )
 {
 	// @see FStatsThreadState::Condense
 	check( CondensedMessages.Num() >= 2 );
@@ -662,8 +608,6 @@ bool FStatsReadFile::PrepareLoading()
 		}
 	}
 
-	State.MarkAsLoaded();
-
 	// Read metadata.
 	TArray<FStatMessage> MetadataMessages;
 	Stream.ReadFNamesAndMetadataMessages( *Reader, MetadataMessages );
@@ -828,7 +772,6 @@ void FStatsReadFile::ReadRawStats()
 
 void FStatsReadFile::ReadRegularStats()
 {
-	uint64 ReadMessages = 0;
 	TArray<FStatMessage> PendingMessages;
 
 	// Buffer used to store the compressed and decompressed data.
@@ -1215,6 +1158,67 @@ void FCommandStatsFile::Stop()
 	}
 }
 
+void FCommandStatsFile::TestLastSaved()
+{
+	const FString& FilePath = FCommandStatsFile::Get().LastFileSaved;
+
+	class FStatsTestReader : public FStatsReadFile
+	{
+		friend struct FStatsReader<FStatsTestReader>;
+		typedef FStatsReadFile Super;
+
+	public:
+		const FStatsLoadedState& GetState() const
+		{
+			return State;
+		}
+
+	protected:
+		/** Initialization constructor. */
+		FStatsTestReader( const TCHAR* InFilename )
+			: FStatsReadFile( InFilename, false )
+		{
+			// Keep the whole history.
+			SetHistoryFrames( MAX_int32 );
+		}
+	};
+
+	TAutoPtr<FStatsTestReader> Instance( FStatsReader<FStatsTestReader>::Create( *FilePath ) );
+
+	extern void DumpHistoryFrame( FStatsThreadState const& StatsData, int64 TargetFrame, float InDumpCull = 0.0f, int32 MaxDepth = MAX_int32, TCHAR const* Filter = NULL );
+
+	if (Instance)
+	{
+		Instance->ReadAndProcessSynchronously();
+		const FStatsLoadedState& Loaded = Instance->GetState();
+
+		if (Loaded.GetLatestValidFrame() < 0)
+		{
+			UE_LOG( LogStats, Log, TEXT( "Failed to stats file: %s" ), *FilePath );
+			return;
+		}
+		UE_LOG( LogStats, Log, TEXT( "Loaded stats file: %s, %lld frame" ), *FilePath, 1 + Loaded.GetLatestValidFrame() - Loaded.GetOldestValidFrame() );
+		{
+			int64 TestFrame = Loaded.GetOldestValidFrame();
+			UE_LOG( LogStats, Log, TEXT( "**************************** Test Frame %lld" ), TestFrame );
+			DumpHistoryFrame( Loaded, TestFrame );
+		}
+	{
+		int64 TestFrame = (Loaded.GetLatestValidFrame() + Loaded.GetOldestValidFrame()) / 2;
+		if (Loaded.IsFrameValid( TestFrame ))
+		{
+			UE_LOG( LogStats, Log, TEXT( "**************************** Test Frame %lld" ), TestFrame );
+			DumpHistoryFrame( Loaded, TestFrame );
+		}
+	}
+	{
+		int64 TestFrame = Loaded.GetLatestValidFrame();
+		UE_LOG( LogStats, Log, TEXT( "**************************** Test Frame %lld" ), TestFrame );
+		DumpHistoryFrame( Loaded, TestFrame );
+	}
+	}
+}
+
 #undef LOCTEXT_NAMESPACE
 
 #endif // STATS
diff --git a/Engine/Source/Runtime/Core/Private/Windows/WindowsPlatformMemory.cpp b/Engine/Source/Runtime/Core/Private/Windows/WindowsPlatformMemory.cpp
--- a/Engine/Source/Runtime/Core/Private/Windows/WindowsPlatformMemory.cpp
+++ b/Engine/Source/Runtime/Core/Private/Windows/WindowsPlatformMemory.cpp
@@ -4,6 +4,9 @@
 
 #include "MallocTBB.h"
 #include "MallocAnsi.h"
+#if USE_MALLOC_STOMP
+	#include "MallocStomp.h"
+#endif // USE_MALLOC_STOMP
 #include "GenericPlatformMemoryPoolStats.h"
 #include "MemoryMisc.h"
 
@@ -72,6 +75,9 @@ FMalloc* FWindowsPlatformMemory::BaseAllocator()
 
 #if FORCE_ANSI_ALLOCATOR
 	return new FMallocAnsi();
+#elif USE_MALLOC_STOMP
+	// Allocator that helps track down memory stomps.
+	return new FMallocStomp();
 #elif (WITH_EDITORONLY_DATA || IS_PROGRAM) && TBB_ALLOCATOR_ALLOWED
 	return new FMallocTBB();
 #else
@@ -164,6 +170,11 @@ const FPlatformMemoryConstants& FWindowsPlatformMemory::GetConstants()
 	return MemoryConstants;	
 }
 
+bool FWindowsPlatformMemory::PageProtect( void* const Ptr, const SIZE_T Size, const uint32 ProtectMode )
+{
+	DWORD flOldProtect;
+	return VirtualProtect(Ptr, Size, ProtectMode, &flOldProtect) != 0;
+}
 void* FWindowsPlatformMemory::BinnedAllocFromOS( SIZE_T Size )
 {
 	return VirtualAlloc( NULL, Size, MEM_COMMIT, PAGE_READWRITE );
diff --git a/Engine/Source/Runtime/Core/Private/Windows/WindowsPlatformStackWalk.cpp b/Engine/Source/Runtime/Core/Private/Windows/WindowsPlatformStackWalk.cpp
--- a/Engine/Source/Runtime/Core/Private/Windows/WindowsPlatformStackWalk.cpp
+++ b/Engine/Source/Runtime/Core/Private/Windows/WindowsPlatformStackWalk.cpp
@@ -699,19 +699,29 @@ bool FWindowsPlatformStackWalk::InitStackWalking()
 		//		SymOpts |= SYMOPT_CASE_INSENSITIVE;
 
 		SymSetOptions( SymOpts );
-
-		// Initialize the symbol engine.
+	
+		// Initialize the symbol engine.		
 		const FString RemoteStorage = GetRemoteStorage(GetDownstreamStorage());
 		SymInitializeW( GetCurrentProcess(), RemoteStorage.IsEmpty() ? nullptr : *RemoteStorage, true );
 	
 		GNeedToRefreshSymbols = false;
 		GStackWalkingInitialized = true;
+
+		if (!FPlatformProperties::IsMonolithicBuild() && FPlatformStackWalk::WantsDetailedCallstacksInNonMonolithicBuilds())
+		{
+			LoadProcessModules( RemoteStorage );
+		}			
 	}
 #if WINVER > 0x502
 	else if (GNeedToRefreshSymbols)
 	{
 		// Refresh and reload symbols
 		SymRefreshModuleList( GetCurrentProcess() );
+		if (!FPlatformProperties::IsMonolithicBuild() && FPlatformStackWalk::WantsDetailedCallstacksInNonMonolithicBuilds())
+		{
+			const FString RemoteStorage = GetRemoteStorage( GetDownstreamStorage() );
+			LoadProcessModules( RemoteStorage );
+		}
 		GNeedToRefreshSymbols = false;
 	}
 #endif
diff --git a/Engine/Source/Runtime/Core/Public/Core.h b/Engine/Source/Runtime/Core/Public/Core.h
--- a/Engine/Source/Runtime/Core/Public/Core.h
+++ b/Engine/Source/Runtime/Core/Public/Core.h
@@ -115,4 +115,5 @@
 #include "Attribute.h"
 #include "Optional.h"
 #include "EnumClassFlags.h"
+#include "EnumRange.h"
 #include "Misc/DisableOldUETypes.h"
diff --git a/Engine/Source/Runtime/Core/Public/GenericPlatform/GenericPlatformMemory.h b/Engine/Source/Runtime/Core/Public/GenericPlatform/GenericPlatformMemory.h
--- a/Engine/Source/Runtime/Core/Public/GenericPlatform/GenericPlatformMemory.h
+++ b/Engine/Source/Runtime/Core/Public/GenericPlatform/GenericPlatformMemory.h
@@ -223,6 +223,16 @@ struct CORE_API FGenericPlatformMemory
 	 */
 	static uint32 GetPhysicalGBRam();
 
+	/**
+	 * Changes the protection on a region of committed pages in the virtual address space.
+	 *
+	 * @param Ptr Address to the starting page of the region of pages whose access protection attributes are to be changed.
+	 * @param Size The size of the region whose access protection attributes are to be changed, in bytes.
+	 * @param ProtectMode The OS dependent memory protection option.
+	 * @return True if the specified pages' protection mode was changed.
+	 */
+	static bool PageProtect( void* const Ptr, const SIZE_T Size, const uint32 ProtectMode );
+
 	/**
 	 * Allocates pages from the OS.
 	 *
diff --git a/Engine/Source/Runtime/Core/Public/GenericPlatform/GenericPlatformStackWalk.h b/Engine/Source/Runtime/Core/Public/GenericPlatform/GenericPlatformStackWalk.h
--- a/Engine/Source/Runtime/Core/Public/GenericPlatform/GenericPlatformStackWalk.h
+++ b/Engine/Source/Runtime/Core/Public/GenericPlatform/GenericPlatformStackWalk.h
@@ -78,6 +78,11 @@ struct CORE_API FGenericPlatformStackWalk
 {
 	typedef FGenericPlatformStackWalk Base;
 
+	/**
+	* Initializes options related to stack walking from ini, i.e. how detailed the stack walking should be, performance settings etc.
+	*/
+	static void Init();
+
 	/**
 	 * Initializes stack traversal and symbol. Must be called before any other stack/symbol functions. Safe to reenter.
 	 */
@@ -182,4 +187,10 @@ struct CORE_API FGenericPlatformStackWalk
 	{
 		return 0;
 	}
+
+protected:
+
+	/** Returns true if non-monolithic builds should produce full callstacks in the log (and load all debug symbols) */
+	static bool WantsDetailedCallstacksInNonMonolithicBuilds();
+
 };
diff --git a/Engine/Source/Runtime/Core/Public/HAL/MallocStomp.h b/Engine/Source/Runtime/Core/Public/HAL/MallocStomp.h
new file mode 100644
--- /dev/null
+++ b/Engine/Source/Runtime/Core/Public/HAL/MallocStomp.h
@@ -0,0 +1,130 @@
+// Copyright 1998-2015 Epic Games, Inc. All Rights Reserved.
+
+#pragma once
+
+//#define USE_MALLOC_STOMP 1
+
+#if USE_MALLOC_STOMP
+
+/**
+ * Stomp memory allocator. It helps find the following errors:
+ * - Read or writes off the end of an allocation.
+ * - Read or writes off the beginning of an allocation.
+ * - Read or writes after freeing an allocation. 
+ */
+class FMallocStomp : public FMalloc
+{
+private:
+#if PLATFORM_64BITS
+	/** Expected value to be found in the sentinel. */
+	static const SIZE_T SentinelExpectedValue = 0xdeadbeefdeadbeef;
+#else
+	/** Expected value to be found in the sentinel. */
+	static const SIZE_T SentinelExpectedValue = 0xdeadbeef;
+#endif
+
+	static const SIZE_T PageSize = 4096U; // TODO: Verify the assumption that all relevant platforms use 4KiB pages.
+
+#if PLATFORM_WINDOWS || PLATFORM_XBOXONE
+	static const uint32 NoAccessProtectMode = PAGE_NOACCESS;
+#elif PLATFORM_LINUX || PLATFORM_MAC
+	static const uint32 NoAccessProtectMode = PROT_NONE;
+#else
+	#error The stomp allocator isn't supported in this platform.
+#endif
+
+	struct FAllocationData
+	{
+		/** Pointer to the full allocation. Needed so the OS knows what to free. */
+		void	*FullAllocationPointer;
+		/** Full size of the allocation including the extra page. */
+		SIZE_T	FullSize;
+		/** Size of the allocation requested. */
+		SIZE_T	Size;
+		/** Sentinel used to check for underrun. */
+		SIZE_T	Sentinel;
+	};
+
+	/** If it is set to true, instead of focusing on overruns the allocator will focus on underruns. */
+	const bool bUseUnderrunMode;
+
+public:
+	// FMalloc interface.
+	FMallocStomp(const bool InUseUnderrunMode = false)
+		: bUseUnderrunMode(InUseUnderrunMode)
+	{}
+
+	/**
+	 * Allocates a block of a given number of bytes of memory with the required alignment.
+	 * In the process it allocates as many pages as necessary plus one that will be protected
+	 * making it unaccessible and causing an exception. The actual allocation will be pushed
+	 * to the end of the last valid unprotected page. To deal with underrun errors a sentinel
+	 * is added right before the allocation in page which is checked on free.
+	 *
+	 * @param Size Size in bytes of the memory block to allocate.
+	 * @param Alignment Alignment in bytes of the memory block to allocate.
+	 * @return A pointer to the beginning of the memory block.
+	 */
+	virtual void* Malloc(SIZE_T Size, uint32 Alignment) override;
+
+	/**
+	 * Changes the size of the memory block pointed to by OldPtr.
+	 * The function may move the memory block to a new location.
+	 *
+	 * @param OldPtr Pointer to a memory block previously allocated with Malloc. 
+	 * @param NewSize New size in bytes for the memory block.
+	 * @param Alignment Alignment in bytes for the reallocation.
+	 * @return A pointer to the reallocated memory block, which may be either the same as ptr or a new location.
+	 */
+	virtual void* Realloc(void* InPtr, SIZE_T NewSize, uint32 Alignment) override;
+
+	/**
+	 * Frees a memory allocation and verifies the sentinel in the process.
+	 *
+	 * @param InPtr Pointer of the data to free.
+	 */
+	virtual void Free(void* InPtr) override;
+
+	/**
+	 * If possible determine the size of the memory allocated at the given address.
+	 * This will included all the pages that were allocated so it will be far more
+	 * than what's set on the FAllocationData.
+	 *
+	 * @param Original - Pointer to memory we are checking the size of
+	 * @param SizeOut - If possible, this value is set to the size of the passed in pointer
+	 * @return true if succeeded
+	 */
+	virtual bool GetAllocationSize(void *Original, SIZE_T &SizeOut) override;
+
+	/**
+	 * Dumps details about all allocations to an output device
+	 *
+	 * @param Ar	[in] Output device
+	 */
+	virtual void DumpAllocatorStats( FOutputDevice& Ar ) override
+	{
+		// No meaningful stats to dump.
+	}
+
+	/**
+	 * Validates the allocator's heap
+	 */
+	virtual bool ValidateHeap() override
+	{
+		// Nothing to do here since validation happens as data is accessed
+		// through page protection, and on each free checking the sentinel.
+		return true;
+	}
+
+	virtual bool Exec( UWorld* InWorld, const TCHAR* Cmd, FOutputDevice& Ar ) override
+	{
+		return false;
+	}
+
+	virtual const TCHAR* GetDescriptiveName() override
+	{
+		return TEXT( "Stomp" );
+	}
+};
+
+#endif // USE_MALLOC_STOMP
\ No newline at end of file
diff --git a/Engine/Source/Runtime/Core/Public/Linux/LinuxPlatformMemory.h b/Engine/Source/Runtime/Core/Public/Linux/LinuxPlatformMemory.h
--- a/Engine/Source/Runtime/Core/Public/Linux/LinuxPlatformMemory.h
+++ b/Engine/Source/Runtime/Core/Public/Linux/LinuxPlatformMemory.h
@@ -49,6 +49,7 @@ struct CORE_API FLinuxPlatformMemory : public FGenericPlatformMemory
 	static class FMalloc* BaseAllocator();
 	static FPlatformMemoryStats GetStats();
 	static const FPlatformMemoryConstants& GetConstants();
+	static bool PageProtect( void* const Ptr, const SIZE_T Size, const uint32 ProtectMode );
 	static void* BinnedAllocFromOS( SIZE_T Size );
 	static void BinnedFreeToOS( void* Ptr );
 	static FSharedMemoryRegion * MapNamedSharedMemoryRegion(const FString& InName, bool bCreate, uint32 AccessMode, SIZE_T Size);
diff --git a/Engine/Source/Runtime/Core/Public/Mac/MacPlatformMemory.h b/Engine/Source/Runtime/Core/Public/Mac/MacPlatformMemory.h
--- a/Engine/Source/Runtime/Core/Public/Mac/MacPlatformMemory.h
+++ b/Engine/Source/Runtime/Core/Public/Mac/MacPlatformMemory.h
@@ -24,6 +24,7 @@ struct CORE_API FMacPlatformMemory : public FGenericPlatformMemory
 	static FPlatformMemoryStats GetStats();
 	static const FPlatformMemoryConstants& GetConstants();
 	static FMalloc* BaseAllocator();
+	static bool PageProtect( void* const Ptr, const SIZE_T Size, const uint32 ProtectMode );
 	static void* BinnedAllocFromOS( SIZE_T Size );
 	static void BinnedFreeToOS( void* Ptr );
 	//~ End FGenericPlatformMemory Interface
diff --git a/Engine/Source/Runtime/Core/Public/Misc/Build.h b/Engine/Source/Runtime/Core/Public/Misc/Build.h
--- a/Engine/Source/Runtime/Core/Public/Misc/Build.h
+++ b/Engine/Source/Runtime/Core/Public/Misc/Build.h
@@ -130,7 +130,7 @@
 * Whether we support hot-reload. Currently requires a non-monolithic build and non-shipping configuration.
 */
 #ifndef WITH_HOT_RELOAD
-	#define WITH_HOT_RELOAD (!IS_MONOLITHIC && !UE_BUILD_SHIPPING && !UE_GAME)
+	#define WITH_HOT_RELOAD (!IS_MONOLITHIC && !UE_BUILD_SHIPPING && !UE_BUILD_TEST && !UE_GAME)
 #endif
 
 /**
diff --git a/Engine/Source/Runtime/Core/Public/Misc/EnumRange.h b/Engine/Source/Runtime/Core/Public/Misc/EnumRange.h
new file mode 100644
--- /dev/null
+++ b/Engine/Source/Runtime/Core/Public/Misc/EnumRange.h
@@ -0,0 +1,186 @@
+// Copyright 1998-2015 Epic Games, Inc. All Rights Reserved.
+
+#pragma once
+
+/**
+ * Defines a contiguous enum range containing Count values, starting from zero:
+ *
+ * Example:
+ *
+ * enum class ECountedThing
+ * {
+ *     First,
+ *     Second,
+ *     Third,
+ *
+ *     Count
+ * };
+ *
+ * // Defines iteration over ECountedThing to be: First, Second, Third
+ * ENUM_RANGE_BY_COUNT(ECountedThing, ECountedThing::Count)
+ */
+#define ENUM_RANGE_BY_COUNT(EnumType, Count) ENUM_RANGE_BY_FIRST_AND_LAST(EnumType, 0, (__underlying_type(EnumType))(Count) - 1)
+
+
+/**
+ * Defines a contiguous enum range with specific start and end values:
+ *
+ * Example:
+ *
+ * enum class EDoubleEndedThing
+ * {
+ *     Invalid,
+ *
+ *     First,
+ *     Second,
+ *     Third,
+ *
+ *     Count
+ * };
+ *
+ * // Defines iteration over EDoubleEndedThing to be: First, Second, Third
+ * ENUM_RANGE_BY_FIRST_AND_LAST(EDoubleEndedThing, EDoubleEndedThing::First, EDoubleEndedThing::Third)
+ */
+#define ENUM_RANGE_BY_FIRST_AND_LAST(EnumType, First, Last) \
+	template <> \
+	struct NEnumRangePrivate::TEnumRangeTraits<EnumType> \
+	{ \
+		enum { RangeType = 0 }; \
+		static const __underlying_type(EnumType) Begin = (__underlying_type(EnumType))(First); \
+		static const __underlying_type(EnumType) End   = (__underlying_type(EnumType))(Last) + 1; \
+	};
+
+
+/**
+ * Defines a non-contiguous enum range with specific individual values:
+ *
+ * Example:
+ *
+ * enum class ERandomValuesThing
+ * {
+ *     First  = 2,
+ *     Second = 3,
+ *     Third  = 5,
+ *     Fourth = 7,
+ *     Fifth  = 11
+ * };
+ *
+ * // Defines iteration over ERandomValuesThing to be: First, Second, Third, Fourth, Fifth
+ * ENUM_RANGE_BY_VALUES(ERandomValuesThing, ERandomValuesThing::First, ERandomValuesThing::Second, ERandomValuesThing::Third, ERandomValuesThing::Fourth, ERandomValuesThing::Fifth)
+ */
+#define ENUM_RANGE_BY_VALUES(EnumType, ...) \
+	template <> \
+	struct NEnumRangePrivate::TEnumRangeTraits<EnumType> \
+	{ \
+		enum { RangeType = 1 }; \
+		template <typename Dummy> \
+		static const EnumType* GetPointer(bool bLast) \
+		{ \
+			static const EnumType Values[] = { __VA_ARGS__ }; \
+			return bLast ? Values + sizeof(Values) / sizeof(EnumType) : Values; \
+		} \
+	};
+
+
+namespace NEnumRangePrivate
+{
+	template <typename EnumType>
+	struct TEnumRangeTraits
+	{
+		enum { RangeType = -1 };
+	};
+
+	template <typename EnumType, int32 RangeType>
+	struct TEnumRange_Impl
+	{
+		static_assert(sizeof(EnumType) == 0, "Unknown enum type - use one of the ENUM_RANGE macros to define iteration semantics for your enum type.");
+	};
+
+	template <typename EnumType>
+	struct TEnumContiguousIterator
+	{
+		typedef __underlying_type(EnumType) IntType;
+
+		FORCEINLINE explicit TEnumContiguousIterator(IntType InValue)
+			: Value(InValue)
+		{
+		}
+
+		FORCEINLINE TEnumContiguousIterator& operator++()
+		{
+			++Value;
+			return *this;
+		}
+
+		FORCEINLINE EnumType operator*() const
+		{
+			return (EnumType)Value;
+		}
+
+	private:
+		IntType Value;
+
+		FORCEINLINE friend bool operator!=(const TEnumContiguousIterator& Lhs, const TEnumContiguousIterator& Rhs)
+		{
+			return Lhs.Value != Rhs.Value;
+		}
+	};
+
+	template <typename EnumType>
+	struct TEnumValueArrayIterator
+	{
+		FORCEINLINE explicit TEnumValueArrayIterator(const EnumType* InPtr)
+			: Ptr(InPtr)
+		{
+		}
+
+		FORCEINLINE TEnumValueArrayIterator& operator++()
+		{
+			++Ptr;
+			return *this;
+		}
+
+		FORCEINLINE EnumType operator*() const
+		{
+			return *Ptr;
+		}
+
+	private:
+		const EnumType* Ptr;
+
+		FORCEINLINE friend bool operator!=(const TEnumValueArrayIterator& Lhs, const TEnumValueArrayIterator& Rhs)
+		{
+			return Lhs.Ptr != Rhs.Ptr;
+		}
+	};
+
+	template <typename EnumType>
+	struct TEnumRange_Impl<EnumType, 0>
+	{
+		friend TEnumContiguousIterator<EnumType> begin(const TEnumRange_Impl&) { return TEnumContiguousIterator<EnumType>(TEnumRangeTraits<EnumType>::Begin); }
+		friend TEnumContiguousIterator<EnumType> end  (const TEnumRange_Impl&) { return TEnumContiguousIterator<EnumType>(TEnumRangeTraits<EnumType>::End  ); }
+	};
+
+	template <typename EnumType>
+	struct TEnumRange_Impl<EnumType, 1>
+	{
+		friend TEnumValueArrayIterator<EnumType> begin(const TEnumRange_Impl&) { return TEnumValueArrayIterator<EnumType>(TEnumRangeTraits<EnumType>::template GetPointer<void>(false)); }
+		friend TEnumValueArrayIterator<EnumType> end  (const TEnumRange_Impl&) { return TEnumValueArrayIterator<EnumType>(TEnumRangeTraits<EnumType>::template GetPointer<void>(true )); }
+	};
+}
+
+/**
+ * Range type for iterating over enum values.  Enums should define themselves as iterable by specifying
+ * one of the ENUM_RANGE_* macros.
+ *
+ * Example:
+ *
+ * for (ECountedThing Val : TEnumRange<ECountedThing>())
+ * {
+ *     ...
+ * }
+ **/
+template <typename EnumType>
+struct TEnumRange : NEnumRangePrivate::TEnumRange_Impl<EnumType, NEnumRangePrivate::TEnumRangeTraits<EnumType>::RangeType>
+{
+};
diff --git a/Engine/Source/Runtime/Core/Public/Misc/StringUtility.h b/Engine/Source/Runtime/Core/Public/Misc/StringUtility.h
--- a/Engine/Source/Runtime/Core/Public/Misc/StringUtility.h
+++ b/Engine/Source/Runtime/Core/Public/Misc/StringUtility.h
@@ -7,12 +7,12 @@
 
 #pragma once
 
-#if !(UE_BUILD_SHIPPING || UE_BUILD_TEST)
+#if !UE_BUILD_SHIPPING
 /** Needed for the console command "DumpConsoleCommands" */
 CORE_API void ConsoleCommandLibrary_DumpLibrary(class UWorld* InWorld, FExec& SubSystem, const FString& Pattern, FOutputDevice& Ar);
 /** Needed for the console command "Help" */
 CORE_API void ConsoleCommandLibrary_DumpLibraryHTML(class UWorld* InWorld, FExec& SubSystem, const FString& OutPath);
-#endif // !(UE_BUILD_SHIPPING || UE_BUILD_TEST)
+#endif // !UE_BUILD_SHIPPING
 
 namespace StringUtility
 {
diff --git a/Engine/Source/Runtime/Core/Public/Stats/Stats2.h b/Engine/Source/Runtime/Core/Public/Stats/Stats2.h
--- a/Engine/Source/Runtime/Core/Public/Stats/Stats2.h
+++ b/Engine/Source/Runtime/Core/Public/Stats/Stats2.h
@@ -1838,8 +1838,8 @@ struct FStat_##StatName\
 #define DECLARE_STATS_GROUP_MAYBE_COMPILED_OUT(GroupDesc, GroupId, GroupCat, CompileIn) \
 	DECLARE_STAT_GROUP(GroupDesc, GroupId, GroupCat, false, CompileIn);
 
-#ifdef UE_BUILD_DEBUG
-	#define SCOPE_CYCLE_COUNTER_GUARD {const char* ReadTheText = "SCOPE_CYCLE_COUNTER can't be used in the global scope.";}
+#if UE_BUILD_DEBUG
+	#define SCOPE_CYCLE_COUNTER_GUARD { const char* ReadTheText = "SCOPE_CYCLE_COUNTER can't be used in the global scope."; ReadTheText = ReadTheText; }
 #else
 	#define SCOPE_CYCLE_COUNTER_GUARD
 #endif // UE_BUILD_DEBUG
