Shadow warning fix.

diff --git a/Engine/Source/Developer/Profiler/Private/ProfilerDataSource.cpp b/Engine/Source/Developer/Profiler/Private/ProfilerDataSource.cpp
--- a/Engine/Source/Developer/Profiler/Private/ProfilerDataSource.cpp
+++ b/Engine/Source/Developer/Profiler/Private/ProfilerDataSource.cpp
@@ -337,7 +337,7 @@ void FEventGraphSample::SetRootAndThreadForAllChildren()
 	}
 }
 
-void FEventGraphSample::FixChildrenTimesAndCalcAveragesForAllChildren( const double NumFrames )
+void FEventGraphSample::FixChildrenTimesAndCalcAveragesForAllChildren( const double InNumFrames )
 {
 	struct FFixChildrenTimesAndCalcAverages
 	{
@@ -347,7 +347,7 @@ void FEventGraphSample::FixChildrenTimesAndCalcAveragesForAllChildren( const dou
 		}
 	};
 
-	ExecuteOperationForAllChildren( FFixChildrenTimesAndCalcAverages(), NumFrames );
+	ExecuteOperationForAllChildren( FFixChildrenTimesAndCalcAverages(), InNumFrames );
 }
 
 FEventGraphSamplePtr FEventGraphSample::FindChildPtr( const FEventGraphSamplePtr& OtherChild )
