Copying //UE4/Dev-Core to //UE4/Main

==========================
MAJOR FEATURES + CHANGES
==========================

Change 2836261 on 2016/01/20 by [REDACTED_PERSON][REDACTED_HANDLE]

	Flush FAsyncPackage cache after pre-load to reduce peak memory usage when async loading (5.5-10x).

Change 2828630 on 2016/01/14 by [REDACTED_PERSON][REDACTED_HANDLE]

	Adding debug code to catch memory stomps in the async loading buffers that's independent from mallocstomp allocator. Changed the signature of PageProtect functions to be able to read-only protect memory.

Change 2816129 on 2016/01/05 by [REDACTED_PERSON][REDACTED_HANDLE]

	Fixes for Realloc and alignment logic which caused redundant reallocations and incorrect binning.

Change 2821054 on 2016/01/08 by [REDACTED_PERSON][REDACTED_HANDLE]

	Further Realloc savings when realigning within a block.

Change 2806820 on 2015/12/17 by [REDACTED_PERSON][REDACTED_HANDLE]

	New AlignDown function, like Align, but which rounds a value/pointer down to the next alignment instead of up.

Change 2806816 on 2015/12/17 by [REDACTED_PERSON][REDACTED_HANDLE]

	Sort UHT modules by type to improve iteration times in conjunction with makefiles.

Change 2823235 on 2016/01/11 by [REDACTED_PERSON][REDACTED_HANDLE]

	UHT error messages about missing GENERATED_BODY() macros updated to represent intended use.

Change 2806815 on 2015/12/17 by [REDACTED_PERSON][REDACTED_HANDLE]

	Module types split into Game and Engine runtime versions.

Change 2833809 on 2016/01/19 by [REDACTED_PERSON][REDACTED_HANDLE]

	StaticLoadObject will now respect LOAD_NoRedirects flag.

Change 2811194 on 2015/12/22 by [REDACTED_PERSON][REDACTED_HANDLE]

	#UE4 Prevent loading packages that have a newer LegacyFileVersion since serialization for FCustomVersion is not forward compatible. UE-24443

Change 2806818 on 2015/12/17 by [REDACTED_PERSON][REDACTED_HANDLE]

	Removal of stats from MallocBinned2, to be readded later.

Change 2807069 on 2015/12/17 by [REDACTED_PERSON][REDACTED_HANDLE]

	Clarification of some bucket hashing terminology.

Change 2815117 on 2016/01/04 by [REDACTED_PERSON][REDACTED_HANDLE]

	Fix for a missing root build path on game modules.

Change 2815673 on 2016/01/05 by [REDACTED_PERSON][REDACTED_HANDLE]

	Move FMalloc verification into a proxy object.

Change 2822873 on 2016/01/11 by [REDACTED_PERSON][REDACTED_HANDLE]

	Fixes to off-by-one errors and removal of BinnedSizeLimit (assumed to be the same as MAX_POOLED_ALLOCATION_SIZE after OBO fix).

Change 2822923 on 2016/01/11 by [REDACTED_PERSON][REDACTED_HANDLE]

	Simplification of MemSizeToPoolTable indexing.

Change 2824974 on 2016/01/12 by [REDACTED_PERSON][REDACTED_HANDLE]

	Assert fixed.
	AllocateBlockFromPool's return value made debuggable.

Change 2825241 on 2016/01/12 by [REDACTED_PERSON][REDACTED_HANDLE]

	UHT now returns an error code on a warning when -warningsaserrors is specified.

Change 2825291 on 2016/01/12 by [REDACTED_PERSON][REDACTED_HANDLE]

	WarningsAsErrors enabled on UHT, after disabling the hardcoded behavior in CL# 2825241.

Change 2829846 on 2016/01/15 by [REDACTED_PERSON][REDACTED_HANDLE]

	GitHub #1938 - wrong Max value of enum is used during net serialization

Change 2829914 on 2016/01/15 by [REDACTED_PERSON][REDACTED_HANDLE]

	Reduce the amount of memory allocated for async cache buffers when guarding against memory stomps.

Change 2829988 on 2016/01/15 by [REDACTED_PERSON][REDACTED_HANDLE]

	Generalized large pool allocations.
	More redundancy removed.

Change 2831935 on 2016/01/18 by [REDACTED_PERSON][REDACTED_HANDLE]

	Added UserActivity property to crash description in CRP and CR website.
	[OR-12043] - Phone Home where crashes occur - pass context info to Crash Reporter

	DB column added to db-09 by ColinR matching this change.
	Published to server on Jan 18th 2016

Change 2834003 on 2016/01/19 by [REDACTED_PERSON][REDACTED_HANDLE]

	Added Linux to normal callstack parsing code on CR website
	[UE-25527] - Linux CrashReporter is missing information

	Published to server on Jan 19th 2016

Change 2835466 on 2016/01/20 by [REDACTED_PERSON][REDACTED_HANDLE]

	Fix issue for cancelling package loads when there are still packages queued.  Call their PackageLoadedDelegate with a "Cancelled" result.

	Should solve remaining issue with UE-24062 - "Calling CancelAsyncLoading triggers an assert in FAsyncPackage::DetachLinker()"
	- (ULevelStreaming::AsyncLevelLoadComplete was not being called if packages were still queued when cancel was issued)

Change 2836803 on 2016/01/20 by [REDACTED_PERSON][REDACTED_HANDLE]

	CrashReportWebsite - fix exception thrown when parsing certain callstack formats

Change 2837952 on 2016/01/21 by [REDACTED_PERSON][REDACTED_HANDLE]

	Changing FAsyncIORequest to be stored as reference when cancelling IO requests to improve performance.

Change 2838289 on 2016/01/21 by [REDACTED_PERSON][REDACTED_HANDLE]

diff --git a/Engine/Plugins/ScriptPlugin/Source/ScriptGeneratorPlugin/Private/ScriptGeneratorPlugin.cpp b/Engine/Plugins/ScriptPlugin/Source/ScriptGeneratorPlugin/Private/ScriptGeneratorPlugin.cpp
--- a/Engine/Plugins/ScriptPlugin/Source/ScriptGeneratorPlugin/Private/ScriptGeneratorPlugin.cpp
+++ b/Engine/Plugins/ScriptPlugin/Source/ScriptGeneratorPlugin/Private/ScriptGeneratorPlugin.cpp
@@ -68,7 +68,7 @@ void FScriptGeneratorPlugin::Initialize(const FString& RootLocalPath, const FStr
 
 bool FScriptGeneratorPlugin::ShouldExportClassesForModule(const FString& ModuleName, EBuildModuleType::Type ModuleType, const FString& ModuleGeneratedIncludeDirectory) const
 { 
-	bool bCanExport = (ModuleType == EBuildModuleType::Runtime || ModuleType == EBuildModuleType::Game);
+	bool bCanExport = (ModuleType == EBuildModuleType::EngineRuntime || ModuleType == EBuildModuleType::GameRuntime);
 	if (bCanExport)
 	{
 		// Only export functions from selected modules
diff --git a/Engine/Source/Developer/AssetTools/Private/AssetRenameManager.cpp b/Engine/Source/Developer/AssetTools/Private/AssetRenameManager.cpp
--- a/Engine/Source/Developer/AssetTools/Private/AssetRenameManager.cpp
+++ b/Engine/Source/Developer/AssetTools/Private/AssetRenameManager.cpp
@@ -544,7 +544,10 @@ void FAssetRenameManager::RenameReferencingStringAssetReferences(const TArray<UP
 	{
 		FStringAssetReferenceRenameSerializer(const FString& InOldAssetPath, const FString& InNewAssetPath)
 			: OldAssetPath(InOldAssetPath), NewAssetPath(InNewAssetPath)
-		{ }
+		{
+			// Mark it as saving to correctly process all references
+			ArIsSaving = true;
+		}
 
 		FArchive& operator<<(FStringAssetReference& Reference) override
 		{
diff --git a/Engine/Source/Developer/CrashDebugHelper/Private/CrashDebugHelper.cpp b/Engine/Source/Developer/CrashDebugHelper/Private/CrashDebugHelper.cpp
--- a/Engine/Source/Developer/CrashDebugHelper/Private/CrashDebugHelper.cpp
+++ b/Engine/Source/Developer/CrashDebugHelper/Private/CrashDebugHelper.cpp
@@ -252,7 +252,7 @@ bool ICrashDebugHelper::SyncModules()
 		return false;
 	}
 
-	// @TODO yrx 2015-02-23 Obsolete, remove after 4.8
+	// #CrashReport: 2015-02-23 Obsolete, remove after 4.8
 	const TCHAR* UESymbols = TEXT( "Rocket/Symbols/" );
 	const bool bHasExecutable = !CrashInfo.ExecutablesPath.IsEmpty();
 	const bool bHasSymbols = !CrashInfo.SymbolsPath.IsEmpty();
diff --git a/Engine/Source/Developer/CrashDebugHelper/Private/Windows/WindowsPlatformStackWalkExt.cpp b/Engine/Source/Developer/CrashDebugHelper/Private/Windows/WindowsPlatformStackWalkExt.cpp
--- a/Engine/Source/Developer/CrashDebugHelper/Private/Windows/WindowsPlatformStackWalkExt.cpp
+++ b/Engine/Source/Developer/CrashDebugHelper/Private/Windows/WindowsPlatformStackWalkExt.cpp
@@ -510,7 +510,7 @@ int FWindowsPlatformStackWalkExt::GetCallstacks()
 				ModuleName = ModuleAndFunction;
 			}
 
-			// #YRX_Crash: 2015-07-24 Add for other platforms
+			// #CrashReport: 2015-07-24 Add this to other platforms
 			// If we find an assert, the actual source file we're interested in is the next one up, so reset the source file found flag
 			if( FunctionName.Len() > 0 )
 			{
diff --git a/Engine/Source/Developer/CrashDebugHelper/Private/Windows/WindowsPlatformStackWalkExt.h b/Engine/Source/Developer/CrashDebugHelper/Private/Windows/WindowsPlatformStackWalkExt.h
--- a/Engine/Source/Developer/CrashDebugHelper/Private/Windows/WindowsPlatformStackWalkExt.h
+++ b/Engine/Source/Developer/CrashDebugHelper/Private/Windows/WindowsPlatformStackWalkExt.h
@@ -40,7 +40,7 @@ struct FWindowsPlatformStackWalkExt
 	/** Check to see if the stack address resides within one of the loaded modules i.e. is it valid?. */
 	bool IsOffsetWithinModules( uint64 Offset );
 
-	/** Extract the system info of the crash from the minidump. */ // @TODO yrx 2015-02-24 Move to crash report client?
+	/** Extract the system info of the crash from the minidump. */
 	void GetSystemInfo();
 
 	/** Extracts the thread info from the minidump. */
diff --git a/Engine/Source/Developer/Profiler/Private/ProfilerDataProvider.cpp b/Engine/Source/Developer/Profiler/Private/ProfilerDataProvider.cpp
--- a/Engine/Source/Developer/Profiler/Private/ProfilerDataProvider.cpp
+++ b/Engine/Source/Developer/Profiler/Private/ProfilerDataProvider.cpp
@@ -6,7 +6,7 @@
 	IDataProvider
 -----------------------------------------------------------------------------*/
 
-void IDataProvider::InternalDuplicate( const IDataProviderRef& DataProvider, const uint32 FrameStartIndex, const uint32 FrameEndIndex )
+void IDataProvider::InternalDuplicate( const IDataProviderRef& DataProvider, const uint32 FrameStartIndex, const uint32 FrameEndIndex ) const
 {
 	for( uint32 FrameIdx = FrameStartIndex; FrameIdx < FrameEndIndex; FrameIdx++ )
 	{
@@ -17,11 +17,6 @@ void IDataProvider::InternalDuplicate( const IDataProviderRef& DataProvider, con
 		{
 			FProfilerSample ProfilerSample( GetCollection()[ SampleIndex ] );
 
-			if( FProfilerSample::IsIndexValid( ProfilerSample._ParentIndex ) )
-			{
-				ProfilerSample._ParentIndex -= SampleStartIndex;
-			}
-
 			for( int32 ChildIdx = 0; ChildIdx < ProfilerSample._ChildrenIndices.Num(); ChildIdx++ )
 			{
 				ProfilerSample._ChildrenIndices[ChildIdx] -= SampleStartIndex;
@@ -88,19 +83,15 @@ const uint32 FArrayDataProvider::AddHierarchicalSample
 	const uint32 InThreadID, 
 	const uint32 InGroupID, 
 	const uint32 InStatID, 
-	const double InStartMS, 
-	const double InDurationMS, 
+	const uint32 InDurationCycles,
 	const uint32 InCallsPerFrame, 
 	const uint32 InParentIndex /*= FProfilerSample::InvalidIndex */
 )
 {
 	uint32 HierSampleIndex = Collection.Num();
-	FProfilerSample HierarchicalSample( InThreadID, InGroupID, InStatID, InStartMS, InDurationMS, InCallsPerFrame, InParentIndex );
-#if CHUNKED_PROFILER_DATA
+	FProfilerSample HierarchicalSample( InThreadID, InGroupID, InStatID, InDurationCycles, InCallsPerFrame );
+
 	Collection.AddElement( HierarchicalSample );
-#else
-	Collection.Add( HierarchicalSample );
-#endif
 
 	if( FProfilerSample::IsIndexValid( InParentIndex ) )
 	{
@@ -123,21 +114,13 @@ void FArrayDataProvider::AddCounterSample
 )
 {
 	FProfilerSample CounterSample( InGroupID, InStatID, InCounter, InProfilerSampleType );
-#if CHUNKED_PROFILER_DATA
 	Collection.AddElement( CounterSample );
-#else
-	Collection.Add( CounterSample );
-#endif
 }
 
 const uint32 FArrayDataProvider::AddDuplicatedSample( const FProfilerSample& ProfilerSample )
 {
 	uint32 DuplicateSampleIndex = Collection.Num();
-#if CHUNKED_PROFILER_DATA
 	Collection.AddElement( ProfilerSample );
-#else
-	Collection.Add( ProfilerSample );
-#endif
 	return DuplicateSampleIndex;
 }
 
@@ -146,10 +129,14 @@ const uint32 FArrayDataProvider::GetNumSamples() const
 	return Collection.Num();
 }
 
-const uint64 FArrayDataProvider::GetMemoryUsage() const 
+const SIZE_T FArrayDataProvider::GetMemoryUsage() const
 {
-	uint64 MemoryUsage = FrameTimes.GetAllocatedSize();
+	SIZE_T MemoryUsage = FrameTimes.GetAllocatedSize();
 	MemoryUsage += FrameIndices.GetAllocatedSize();
+	MemoryUsage += ElapsedFrameTimes.GetAllocatedSize();
+	MemoryUsage += FrameCounters.GetAllocatedSize();
+	MemoryUsage += AccumulatedFrameCounters.GetAllocatedSize();
+
 	MemoryUsage += Collection.GetAllocatedSize();
 	MemoryUsage += ChildrenIndicesMemoryUsage;
 	return MemoryUsage;
diff --git a/Engine/Source/Developer/Profiler/Private/ProfilerDataProvider.h b/Engine/Source/Developer/Profiler/Private/ProfilerDataProvider.h
--- a/Engine/Source/Developer/Profiler/Private/ProfilerDataProvider.h
+++ b/Engine/Source/Developer/Profiler/Private/ProfilerDataProvider.h
@@ -34,7 +34,7 @@ protected:
 	 * @return A shared reference with a new instance of the data provider.
 	 */
 	template< class TDataProviderType >
-	TSharedRef<IDataProvider> Duplicate( const uint32 FrameStartIndex, const uint32 NumFrames = 1 )
+	TSharedRef<IDataProvider> Duplicate( const uint32 FrameStartIndex, const uint32 NumFrames = 1 ) const
 	{
 		const uint32 TotalNumFrames = GetNumFrames();
 		check( NumFrames <= TotalNumFrames )
@@ -53,7 +53,7 @@ protected:
 
 protected:
 	/** Helper method used to create a copy of a specified data provider. Removes the template parameters, so can be implemented in the source file. */
-	void InternalDuplicate( const IDataProviderRef& DataProvider, const uint32 FrameStartIndex, const uint32 FrameEndIndex );
+	void InternalDuplicate( const IDataProviderRef& DataProvider, const uint32 FrameStartIndex, const uint32 FrameEndIndex ) const;
 
 public:
 	/**
@@ -73,8 +73,7 @@ public:
 		const uint32 InThreadID, 
 		const uint32 InGroupID, 
 		const uint32 InStatID, 
-		const double InStartMS, 
-		const double InDurationMS, 
+		const uint32 InDurationCycles, 
 		const uint32 InCallsPerFrame, 
 		const uint32 InParentIndex = FProfilerSample::InvalidIndex 
 	) = 0;
@@ -185,7 +184,7 @@ public:
 		return ElapsedFrameTimes[ InFrameIndex ];
 	}
  
-	virtual const uint64 GetMemoryUsage() const = 0;
+	virtual const SIZE_T GetMemoryUsage() const = 0;
 		
 protected:
 	/** An array of indices to the frame's range. */
@@ -239,8 +238,7 @@ public:
 		const uint32 InThreadID, 
 		const uint32 InGroupID, 
 		const uint32 InStatID, 
-		const double InStartMS, 
-		const double InDurationMS, 
+		const uint32 InDurationCycles,
 		const uint32 InCallsPerFrame, 
 		const uint32 InParentIndex = FProfilerSample::InvalidIndex 
 	) override;
@@ -257,7 +255,7 @@ public:
 
 	virtual const uint32 GetNumSamples() const override;
 
-	virtual const uint64 GetMemoryUsage() const override;;
+	virtual const SIZE_T GetMemoryUsage() const override;;
 
 	virtual const FProfilerSampleArray& GetCollection() const override
 	{
diff --git a/Engine/Source/Developer/Profiler/Private/ProfilerDataSource.cpp b/Engine/Source/Developer/Profiler/Private/ProfilerDataSource.cpp
--- a/Engine/Source/Developer/Profiler/Private/ProfilerDataSource.cpp
+++ b/Engine/Source/Developer/Profiler/Private/ProfilerDataSource.cpp
@@ -45,36 +45,44 @@ FGraphDataSource::FGraphDataSource( const FProfilerSessionRef& InProfilerSession
 	}
 }
 
-const TGraphDataType FGraphDataSource::GetUncachedValueFromIndex( const uint32 Index ) const
+const TGraphDataType FGraphDataSource::GetUncachedValueFromIndex( const uint32 FrameIndex ) const
 {
-	check( Index < ProfilerSession->GetDataProvider()->GetNumFrames() );
+	check( FrameIndex < ProfilerSession->GetDataProvider()->GetNumFrames() );
 	double Result = 0.0;
 
-	const IDataProviderRef DataProvider = ProfilerSession->GetDataProvider();
+	// Hierarchical samples are stored in different location.
+	// We skip hierarchical samples to ignore misleading recursion which would be counted twice etc.
+	if (GetSampleType() == EProfilerSampleTypes::HierarchicalTime)
+	{
+		const TMap<uint32, FInclusiveTime>& InclusiveAggregates = ProfilerSession->GetInclusiveAggregateStackStats( FrameIndex );
+		const FInclusiveTime* InclusiveTime = InclusiveAggregates.Find( GetStatID() );
 
-	const FIntPoint& IndicesForFrame = DataProvider->GetSamplesIndicesForFrame( Index );
-	const uint32 SampleStartIndex = IndicesForFrame.X;
-	const uint32 SampleEndIndex = IndicesForFrame.Y;
+		if (InclusiveTime)
+		{
+			Result = ProfilerSession->GetMetaData()->ConvertCyclesToMS( InclusiveTime->DurationCycles ) * Scale;
+		}
+	}
+	else
+	{
+		const IDataProviderRef DataProvider = ProfilerSession->GetDataProvider();
 
-	const FProfilerSampleArray& Collection = DataProvider->GetCollection();
+		const FIntPoint& IndicesForFrame = DataProvider->GetSamplesIndicesForFrame( FrameIndex );
+		const uint32 SampleStartIndex = IndicesForFrame.X;
+		const uint32 SampleEndIndex = IndicesForFrame.Y;
 
-	for( uint32 SampleIndex = SampleStartIndex; SampleIndex < SampleEndIndex; SampleIndex++ )
-	{
-		const FProfilerSample& ProfilerSample = Collection[ SampleIndex ];
-		const bool bValidStat = ProfilerSample.StatID() == GetStatID();
+		const FProfilerSampleArray& Collection = DataProvider->GetCollection();
 
-		if( bValidStat )
+		for (uint32 SampleIndex = SampleStartIndex; SampleIndex < SampleEndIndex; SampleIndex++)
 		{
-			if( GetSampleType() == EProfilerSampleTypes::HierarchicalTime )
-			{
-				Result += ProfilerSample.DurationMS() * Scale;
-			}
-			else 
+			const FProfilerSample& ProfilerSample = Collection[SampleIndex];
+			const bool bValidStat = ProfilerSample.StatID() == GetStatID();
+
+			if (bValidStat)
 			{
-				Result += ProfilerSample.CounterAsFloat() * Scale;
+				Result += ProfilerSample.GetDoubleValue() * Scale;
 			}
 		}
-	}	
+	}
 
 	return (TGraphDataType)Result;
 }
@@ -229,7 +237,6 @@ FEventProperty FEventGraphSample::Properties[ EEventPropertyIndex::InvalidOrMax
 	FEventProperty( EEventPropertyIndex::ThreadPct, TEXT( "ThreadPct" ), STRUCT_OFFSET( FEventGraphSample, _ThreadPct ), EEventPropertyFormatters::TimePct ),
 	FEventProperty( EEventPropertyIndex::FramePct, TEXT( "FramePct" ), STRUCT_OFFSET( FEventGraphSample, _FramePct ), EEventPropertyFormatters::TimePct ),
 	FEventProperty( EEventPropertyIndex::ThreadToFramePct, TEXT( "ThreadToFramePct" ), STRUCT_OFFSET( FEventGraphSample, _ThreadToFramePct ), EEventPropertyFormatters::TimePct ),
-	FEventProperty( EEventPropertyIndex::StartTimeMS, TEXT( "StartTimeMS" ), STRUCT_OFFSET( FEventGraphSample, _StartTimeMS ), EEventPropertyFormatters::TimeMS ),
 	FEventProperty( EEventPropertyIndex::GroupName, TEXT( "GroupName" ), STRUCT_OFFSET( FEventGraphSample, _GroupName ), EEventPropertyFormatters::Name ),
 
 	// Booleans
@@ -274,7 +281,6 @@ void FEventGraphSample::InitializePropertyManagement()
 			.Add( TEXT( "ThreadPct" ), &Properties[EEventPropertyIndex::ThreadPct] )
 			.Add( TEXT( "FramePct" ), &Properties[EEventPropertyIndex::FramePct] )
 			.Add( TEXT( "ThreadToFramePct" ), &Properties[EEventPropertyIndex::ThreadToFramePct] )
-			.Add( TEXT( "StartTimeMS" ), &Properties[EEventPropertyIndex::StartTimeMS] )
 			.Add( TEXT( "GroupName" ), &Properties[EEventPropertyIndex::GroupName] )
 
 			// Booleans
@@ -399,10 +405,9 @@ FEventGraphData::FEventGraphData()
 	, FrameEndIndex( 0 )
 {}
 
-FEventGraphData::FEventGraphData( const FProfilerSessionRef& InProfilerSession, const uint32 InFrameIndex )
+FEventGraphData::FEventGraphData( const FProfilerSession * const InProfilerSession, const uint32 InFrameIndex )
 	: FrameStartIndex( InFrameIndex )
 	, FrameEndIndex( InFrameIndex+1 )
-	, ProfilerSessionPtr( InProfilerSession )
 {
 	static FTotalTimeAndCount Current(0.0f, 0);
 	PROFILER_SCOPE_LOG_TIME( TEXT( "FEventGraphData::FEventGraphData" ), &Current );
@@ -421,7 +426,8 @@ FEventGraphData::FEventGraphData( const FProfilerSessionRef& InProfilerSession,
 	PopulateHierarchy_Recurrent( InProfilerSession, RootEvent, RootProfilerSample, DataProvider );
 
 	// Root sample contains FrameDurationMS
-	RootEvent->_InclusiveTimeMS = RootProfilerSample.DurationMS();
+	const FProfilerStatMetaDataRef& MetaData = InProfilerSession->GetMetaData();
+	RootEvent->_InclusiveTimeMS = MetaData->ConvertCyclesToMS( RootProfilerSample.GetDurationCycles() );
 	RootEvent->_MaxInclusiveTimeMS = RootEvent->_MinInclusiveTimeMS = RootEvent->_AvgInclusiveTimeMS = RootEvent->_InclusiveTimeMS;
 	RootEvent->_InclusiveTimePct = 100.0f;
 
@@ -436,13 +442,13 @@ FEventGraphData::FEventGraphData( const FProfilerSessionRef& InProfilerSession,
 
 void FEventGraphData::PopulateHierarchy_Recurrent
 ( 
-	const FProfilerSessionRef& ProfilerSession,
+	const FProfilerSession * const ProfilerSession,
 	const FEventGraphSamplePtr ParentEvent, 
 	const FProfilerSample& ParentSample, 
 	const IDataProviderRef DataProvider
 )
 {
-	const FProfilerStatMetaDataRef MetaData = ProfilerSession->GetMetaData();
+	const FProfilerStatMetaDataRef& MetaData = ProfilerSession->GetMetaData();
 
 	for( int32 ChildIndex = 0; ChildIndex < ParentSample.ChildrenIndices().Num(); ChildIndex++ )
 	{
@@ -458,7 +464,7 @@ void FEventGraphData::PopulateHierarchy_Recurrent
 		FEventGraphSample* ChildEvent = new FEventGraphSample
 		(
 			ThreadName,	GroupName, ChildSample.StatID(), StatName, 
-			ChildSample.StartMS(), ChildSample.DurationMS(), ChildSample.CounterAsFloat(),
+			MetaData->ConvertCyclesToMS( ChildSample.GetDurationCycles() ), (double)ChildSample.GetCallCount(),
 			ParentEvent
 		);
 
diff --git a/Engine/Source/Developer/Profiler/Private/ProfilerDataSource.h b/Engine/Source/Developer/Profiler/Private/ProfilerDataSource.h
--- a/Engine/Source/Developer/Profiler/Private/ProfilerDataSource.h
+++ b/Engine/Source/Developer/Profiler/Private/ProfilerDataSource.h
@@ -21,12 +21,6 @@ typedef TSharedPtr<class FEventGraphData, ESPMode::ThreadSafe> FEventGraphDataPt
 /** Type definition for shared references to instances of FEventGraphData. */
 typedef TSharedRef<class FEventGraphData, ESPMode::ThreadSafe> FEventGraphDataRef;
 
-/** Type definition for shared pointers to instances of FEventGraphDataHandler. */
-typedef TSharedPtr<class FEventGraphDataHandler> FEventGraphDataHandlerPtr;
-
-/** Type definition for shared references to instances of FEventGraphDataHandler. */
-typedef TSharedRef<class FEventGraphDataHandler> FEventGraphDataHandlerRef;
-
 /** Scratch buffers for multithreaded usage. */
 struct FProfilerScratchArea : public TThreadSingleton<FProfilerScratchArea>
 {
@@ -527,7 +521,7 @@ public:
 
 protected:
 	/* @return a sample value for the specified frame index from the data provider. */
-	const TGraphDataType GetUncachedValueFromIndex( const uint32 Index ) const;
+	const TGraphDataType GetUncachedValueFromIndex( const uint32 FrameIndex ) const;
 
 	/* @return an approximated sample value for the specified time range from the data provider. */
 	const TGraphDataType GetUncachedValueFromTimeRange( const float StartTimeMS, const float EndTimeMS ) const;
@@ -712,7 +706,6 @@ namespace EEventPropertyIndex
 		ThreadPct,
 		FramePct,
 		ThreadToFramePct,
-		StartTimeMS,
 		GroupName,
 		
 		// Booleans
@@ -986,7 +979,6 @@ public:
 		, _GroupName( InName )
 		, _StatName( InName )
 		, _StatID( 0 )
-		, _StartTimeMS( 0.0 )
 		, _InclusiveTimeMS( 0.0 )
 		, _InclusiveTimePct( 0.0 )
 		, _MinInclusiveTimeMS( TNumericLimits<double>::Max() )
@@ -1028,7 +1020,6 @@ protected:
 		const uint32 InStatID,
 		const FName InStatName,
 
-		const double InStartTimeMS,
 		const double InInclusiveTimeMS,
 		const double InNumCallsPerFrame,
 
@@ -1041,7 +1032,6 @@ protected:
 		, _GroupName( InGroupName )
 		, _StatName( InStatName )
 		, _StatID( InStatID )
-		, _StartTimeMS( InStartTimeMS )
 		, _InclusiveTimeMS( InInclusiveTimeMS )
 		, _InclusiveTimePct( 0.0f )
 		, _MinInclusiveTimeMS( InInclusiveTimeMS )
@@ -1073,7 +1063,6 @@ protected:
 		, _GroupName( SourceEvent._GroupName )
 		, _StatName( SourceEvent._StatName )
 		, _StatID( SourceEvent._StatID )
-		, _StartTimeMS( SourceEvent._StartTimeMS )
 		, _InclusiveTimeMS( SourceEvent._InclusiveTimeMS )
 		, _InclusiveTimePct( SourceEvent._InclusiveTimePct )
 		, _MinInclusiveTimeMS( SourceEvent._MinInclusiveTimeMS )
@@ -1118,11 +1107,11 @@ public:
 		//_ExclusiveTimeMS += Other->_ExclusiveTimeMS;
 
 		// Min/Max
-		_MinInclusiveTimeMS = FMath::Min( _MinInclusiveTimeMS, Other->_InclusiveTimeMS );
-		_MaxInclusiveTimeMS = FMath::Max( _MaxInclusiveTimeMS, Other->_InclusiveTimeMS );
+		_MinInclusiveTimeMS = FMath::Min( _MinInclusiveTimeMS, Other->_MinInclusiveTimeMS );
+		_MaxInclusiveTimeMS = FMath::Max( _MaxInclusiveTimeMS, Other->_MaxInclusiveTimeMS );
 
-		_MinNumCallsPerFrame = FMath::Min( _MinNumCallsPerFrame, Other->_NumCallsPerFrame );
-		_MaxNumCallsPerFrame = FMath::Max( _MaxNumCallsPerFrame, Other->_NumCallsPerFrame );
+		_MinNumCallsPerFrame = FMath::Min( _MinNumCallsPerFrame, Other->_MinNumCallsPerFrame );
+		_MaxNumCallsPerFrame = FMath::Max( _MaxNumCallsPerFrame, Other->_MaxNumCallsPerFrame );
 	}
 
 protected:
@@ -1738,10 +1727,6 @@ public:
 	/** Stat ID of the event. */ // TODO: Not needed, move to _statname if meta data is not available
 	uint32 _StatID;
 
-	/** Start time of this event, in milliseconds. */
-	double _StartTimeMS;
-
-
 	/** Duration of this event and its children, in milliseconds. */
 	double _InclusiveTimeMS;
 
@@ -2091,37 +2076,23 @@ struct EEventGraphTypes
 	static FString ToDescription( const EEventGraphTypes::Type EventGraphType );
 };
 
-/** POD helper class for storing information like type and frame indices about the event graph. */
-class FEventGraphDataProperties
+/** Simple struct used to return a result. */
+struct FEventGraphContainer
 {
-public:
-	FEventGraphDataProperties( uint32 InFrameStartIndex, uint32 InFrameEndIndex, const EEventGraphTypes::Type InEventGraphType )
+	/** Initialization constructor. */
+	FEventGraphContainer( uint32 InFrameStartIndex, uint32 InFrameEndIndex, const FEventGraphDataRef& InAverage, const FEventGraphDataRef& InMaximum, const FEventGraphDataRef& InTotal )
 		: FrameStartIndex( InFrameStartIndex )
 		, FrameEndIndex( InFrameEndIndex )
-		, EventGraphType( InEventGraphType )
+		, Average( InAverage )
+		, Maximum( InMaximum )
+		, Total( InTotal )
 	{}
 
-	uint32 FrameStartIndex;
-	uint32 FrameEndIndex;
-	EEventGraphTypes::Type EventGraphType;
-};
-
-template<> struct TIsPODType<FEventGraphDataProperties> { enum { Value = true }; };
-
-/** Helper class used to initialize event graph samples. */
-class FEventGraphDataHandler : public TSharedFromThis<FEventGraphDataHandler>
-{
-public:
-	FEventGraphDataHandler()
-	{}
-
-	FEventGraphDataHandler( const FGuid SessionInstanceID, const uint32 FrameStartIndex, const uint32 FrameEndIndex, const EEventGraphTypes::Type EventGraphType )
-	{
-		EventGraphProperties.Add( SessionInstanceID, FEventGraphDataProperties(FrameStartIndex,FrameEndIndex,EventGraphType) );
-	}
-
-	/** Contains information about frames that require opening an event graphs, stored as SessionInstanceID -> StartIndex. */
-	TMap<FGuid,FEventGraphDataProperties> EventGraphProperties;
+	const uint32 FrameStartIndex;
+	const uint32 FrameEndIndex;
+	FEventGraphDataRef Average;
+	FEventGraphDataRef Maximum;
+	FEventGraphDataRef Total;
 };
 
 /** 
@@ -2138,22 +2109,17 @@ public:
 
 	/** Destructor, for debugging purpose only. */
 	FORCEINLINE_DEBUGGABLE ~FEventGraphData()
-	{
-#if	_DEBUG
-		int k=0;k++;
-		RootEvent;
-#endif // _DEBUG
-	}
+	{}
 
 protected:
 	/**
 	 * Initialization constructor, hidden on purpose, may be only called from the FProfilerSession class.
 	 * 
-	 * @param InProfilerSession	- a const reference to the profiler session that will be used to generate this event graph data 
+	 * @param InProfilerSession	- a const pointer to the profiler session that will be used to generate this event graph data 
 	 * @param InFrameIndex		- the frame number from which to generate this event graph data
 	 *
 	 */
-	FEventGraphData( const FProfilerSessionRef& InProfilerSession, const uint32 InFrameIndex );
+	FEventGraphData( const FProfilerSession * const InProfilerSession, const uint32 InFrameIndex );
 
 	/** Copy constructor, create a full duplication of the source event graph data. */
 	FEventGraphData( const FEventGraphData& Source );
@@ -2171,7 +2137,7 @@ protected:
 	 */
 	void PopulateHierarchy_Recurrent
 	( 
-		const FProfilerSessionRef& ProfilerSession,
+		const FProfilerSession * const ProfilerSession,
 		const FEventGraphSamplePtr ParentEvent, 
 		const FProfilerSample& ParentSample, 
 		const IDataProviderRef DataProvider
@@ -2238,14 +2204,6 @@ public:
 		return Description;
 	}
 
-	/**
-	 * @return the profiler session that was used to generate this event graph data, may be null.
-	 */
-	FProfilerSessionPtr GetProfilerSession() const
-	{
-		return ProfilerSessionPtr.Pin();
-	}
-
 protected:
 	// TODO: FEventGraphSamplePtr -> FEventGraphSampleRef
 
@@ -2260,9 +2218,6 @@ protected:
 
 	/** The frame end index this event graph data was generated from. */
 	uint32 FrameEndIndex;
-
-	/** A weak pointer to the profiler session that was used to generate this event graph data. */
-	FProfilerSessionWeak ProfilerSessionPtr;
 };
 
 //}//namespace FEventGraphSample
diff --git a/Engine/Source/Developer/Profiler/Private/ProfilerFPSAnalyzer.cpp b/Engine/Source/Developer/Profiler/Private/ProfilerFPSAnalyzer.cpp
--- a/Engine/Source/Developer/Profiler/Private/ProfilerFPSAnalyzer.cpp
+++ b/Engine/Source/Developer/Profiler/Private/ProfilerFPSAnalyzer.cpp
@@ -62,6 +62,12 @@ void FFPSAnalyzer::AddSample( float FPSSample )
 	AveFPS = (FPSSample + (float)(Samples.Num()-1) * AveFPS) / (float)Samples.Num();
 }
 
+SIZE_T FFPSAnalyzer::GetMemoryUsage() const
+{
+	const SIZE_T MemoryAllocated = Samples.GetAllocatedSize() + Histogram.GetAllocatedSize();
+	return MemoryAllocated;
+}
+
 int32 FFPSAnalyzer::GetTotalCount()
 {
 	return Samples.Num();
diff --git a/Engine/Source/Developer/Profiler/Private/ProfilerFPSAnalyzer.h b/Engine/Source/Developer/Profiler/Private/ProfilerFPSAnalyzer.h
--- a/Engine/Source/Developer/Profiler/Private/ProfilerFPSAnalyzer.h
+++ b/Engine/Source/Developer/Profiler/Private/ProfilerFPSAnalyzer.h
@@ -59,6 +59,8 @@ public:
 
 	void AddSample( float FPSSample );
 
+	SIZE_T GetMemoryUsage() const;
+
 	TArray<float> Samples;
 	TArray<FProfilerFPSChartEntry> Histogram;
 	float MinFPS;
diff --git a/Engine/Source/Developer/Profiler/Private/ProfilerManager.cpp b/Engine/Source/Developer/Profiler/Private/ProfilerManager.cpp
--- a/Engine/Source/Developer/Profiler/Private/ProfilerManager.cpp
+++ b/Engine/Source/Developer/Profiler/Private/ProfilerManager.cpp
@@ -99,7 +99,7 @@ void FProfilerManager::PostConstructor()
 {
 	// Register tick functions.
 	OnTick = FTickerDelegate::CreateSP( this, &FProfilerManager::Tick );
-	OnTickHandle = FTicker::GetCoreTicker().AddTicker( OnTick );
+	OnTickHandle = FTicker::GetCoreTicker().AddTicker( OnTick, 1.0f );
 
 	// Create profiler client.
 	ProfilerClient = FModuleManager::GetModuleChecked<IProfilerClientModule>("ProfilerClient").CreateProfilerClient();
@@ -221,7 +221,7 @@ void FProfilerManager::LoadProfilerCapture( const FString& ProfilerCaptureFilepa
 
 void FProfilerManager::LoadRawStatsFile( const FString& RawStatsFileFileath )
 {
-	// #YRX_Rework 2015-11-30 
+	// #Profiler: Rework
 #if	0
 	if (ActiveSession.IsValid())
 	{
@@ -260,15 +260,23 @@ bool FProfilerManager::Tick( float DeltaTime )
 {
 	SCOPE_CYCLE_COUNTER(STAT_PM_Tick);
 
-#if	0
-	static int32 NumIdleTicks = 0;
-	NumIdleTicks++;
-	if( NumIdleTicks > 60 )
+	if (ProfilerSession.IsValid() && !bHasCaptureFileFullyProcessed)
 	{
-		LoadRawStatsFile( TEXT( "U:/P4EPIC2/UE4/QAGame/Saved/Profiling/UnrealStats/RenderTestMap-Windows-03.19-13.49.04/RenderTestMap-Windows-19-13.52.19.ue4statsraw" ) );
-		return false;
+		static SIZE_T StartUsedPhysical = FPlatformMemory::GetStats().UsedPhysical;
+
+		const double MBInv = 1.0 / 1024.0 / 1024.0;
+		const SIZE_T DiffPhys = FPlatformMemory::GetStats().UsedPhysical - StartUsedPhysical;
+
+		const double SessionMemory = MBInv*ProfilerSession->GetMemoryUsage();
+		const double PhysMemory = MBInv*DiffPhys;
+
+		UE_LOG( LogStats, Log, TEXT( "ProfilerSession: %6.2f MB (%6.2f MB) # (%6.2f MB) / %7u -> %4u" ), 
+			SessionMemory,
+			PhysMemory,
+			PhysMemory - SessionMemory,
+			ProfilerSession->GetDataProvider()->GetNumSamples(),
+			ProfilerSession->GetDataProvider()->GetNumFrames() );
 	}
-#endif // 0
 
 	return true;
 }
@@ -601,13 +609,12 @@ void FProfilerManager::CloseAllEventGraphTabs()
 
 void FProfilerManager::DataGraph_OnSelectionChangedForIndex( uint32 FrameStartIndex, uint32 FrameEndIndex )
 {
-	PROFILER_SCOPE_LOG_TIME( TEXT( "FProfilerManager::DataGraph_OnSelectionChangedForIndex" ), nullptr );
+//	SCOPE_LOG_TIME_FUNC();
 
 	if (ProfilerSession.IsValid())
 	{
-		FEventGraphDataRef EventGraphDataAverage = ProfilerSession->CreateEventGraphData( FrameStartIndex, FrameEndIndex, EEventGraphTypes::Average );
-		FEventGraphDataRef EventGraphDataMaximum = ProfilerSession->CreateEventGraphData( FrameStartIndex, FrameEndIndex, EEventGraphTypes::Maximum );
-		GetProfilerWindow()->UpdateEventGraph( ProfilerSession->GetInstanceID(), EventGraphDataAverage, EventGraphDataMaximum, false );
+		FEventGraphContainer EventGraphContainer = ProfilerSession->CreateEventGraphData( FrameStartIndex, FrameEndIndex );
+		GetProfilerWindow()->UpdateEventGraph( ProfilerSession->GetInstanceID(), EventGraphContainer.Average, EventGraphContainer.Maximum, false );
 	}
 }
 
diff --git a/Engine/Source/Developer/Profiler/Private/ProfilerPrivatePCH.h b/Engine/Source/Developer/Profiler/Private/ProfilerPrivatePCH.h
--- a/Engine/Source/Developer/Profiler/Private/ProfilerPrivatePCH.h
+++ b/Engine/Source/Developer/Profiler/Private/ProfilerPrivatePCH.h
@@ -23,7 +23,7 @@
 // TODO: Remove dependencies between generic profiler code and profiler ui code
 
 // Common headers and managers.
-#include "ProfilerRawStats.h"
+#include "ProfilerRawStatsForMemory.h"
 #include "ProfilerStream.h"
 
 #include "ProfilerSample.h"
@@ -31,6 +31,7 @@
 #include "ProfilerDataProvider.h"
 #include "ProfilerDataSource.h"
 #include "ProfilerSession.h"
+#include "ProfilerRawStatsForThreadView.h"
 #include "ProfilerCommands.h"
 #include "ProfilerManager.h"
 
diff --git a/Engine/Source/Developer/Profiler/Private/ProfilerRawStatsForMemory.cpp b/Engine/Source/Developer/Profiler/Private/ProfilerRawStatsForMemory.cpp
new file mode 100644
--- /dev/null
+++ b/Engine/Source/Developer/Profiler/Private/ProfilerRawStatsForMemory.cpp
@@ -0,0 +1,884 @@
+// Copyright 1998-2016 Epic Games, Inc. All Rights Reserved.
+
+#include "ProfilerPrivatePCH.h"
+#include "DiagnosticTable.h"
+
+/*-----------------------------------------------------------------------------
+	Sort helpers
+-----------------------------------------------------------------------------*/
+
+/** Sorts allocations by size. */
+struct FAllocationInfoSequenceTagLess
+{
+	FORCEINLINE bool operator()( const FAllocationInfo& A, const FAllocationInfo& B ) const
+	{
+		return A.SequenceTag < B.SequenceTag;
+	}
+};
+
+/** Sorts allocations by size. */
+struct FAllocationInfoSizeGreater
+{
+	FORCEINLINE bool operator()( const FAllocationInfo& A, const FAllocationInfo& B ) const
+	{
+		return B.Size < A.Size;
+	}
+};
+
+/** Sorts combined allocations by size. */
+struct FCombinedAllocationInfoSizeGreater
+{
+	FORCEINLINE bool operator()( const FCombinedAllocationInfo& A, const FCombinedAllocationInfo& B ) const
+	{
+		return B.Size < A.Size;
+	}
+};
+
+
+/** Sorts node allocations by size. */
+struct FNodeAllocationInfoSizeGreater
+{
+	FORCEINLINE bool operator()( const FNodeAllocationInfo& A, const FNodeAllocationInfo& B ) const
+	{
+		return B.Size < A.Size;
+	}
+};
+
+/*-----------------------------------------------------------------------------
+	Callstack decoding/encoding
+-----------------------------------------------------------------------------*/
+
+/** Helper struct used to manipulate stats based callstacks. */
+struct FStatsCallstack
+{
+	/** Separator. */
+	static const TCHAR* CallstackSeparator;
+
+	/** Encodes decoded callstack a string, to be like '45+656+6565'. */
+	static FString Encode( const TArray<FName>& Callstack )
+	{
+		FString Result;
+		for (const auto& Name : Callstack)
+		{
+			Result += TTypeToString<int32>::ToString( (int32)Name.GetComparisonIndex() );
+			Result += CallstackSeparator;
+		}
+		return Result;
+	}
+
+	/** Decodes encoded callstack to an array of FNames. */
+	static void DecodeToNames( const FName& EncodedCallstack, TArray<FName>& out_DecodedCallstack )
+	{
+		TArray<FString> DecodedCallstack;
+		DecodeToStrings( EncodedCallstack, DecodedCallstack );
+
+		// Convert back to FNames
+		for (const auto& It : DecodedCallstack)
+		{
+			NAME_INDEX NameIndex = 0;
+			TTypeFromString<NAME_INDEX>::FromString( NameIndex, *It );
+			const FName LongName = FName( NameIndex, NameIndex, 0 );
+
+			out_DecodedCallstack.Add( LongName );
+		}
+	}
+
+	/** Converts the encoded callstack into human readable callstack. */
+	static FString GetHumanReadable( const FName& EncodedCallstack )
+	{
+		TArray<FName> DecodedCallstack;
+		DecodeToNames( EncodedCallstack, DecodedCallstack );
+		const FString Result = GetHumanReadable( DecodedCallstack );
+		return Result;
+	}
+
+	/** Converts the encoded callstack into human readable callstack. */
+	static FString GetHumanReadable( const TArray<FName>& DecodedCallstack )
+	{
+		FString Result;
+
+		const int32 NumEntries = DecodedCallstack.Num();
+		//for (int32 Index = DecodedCallstack.Num() - 1; Index >= 0; --Index)
+		for (int32 Index = 0; Index < NumEntries; ++Index)
+		{
+			const FName LongName = DecodedCallstack[Index];
+			const FString ShortName = FStatNameAndInfo::GetShortNameFrom( LongName ).ToString();
+			//const FString Group = FStatNameAndInfo::GetGroupNameFrom( LongName ).ToString();
+			FString Desc = FStatNameAndInfo::GetDescriptionFrom( LongName );
+			Desc.Trim();
+
+			if (Desc.Len() == 0)
+			{
+				Result += ShortName;
+			}
+			else
+			{
+				Result += Desc;
+			}
+
+			if (Index != NumEntries - 1)
+			{
+				Result += TEXT( " -> " );
+			}
+		}
+
+		Result.ReplaceInline( TEXT( "STAT_" ), TEXT( "" ), ESearchCase::CaseSensitive );
+		return Result;
+	}
+
+protected:
+	/** Decodes encoded callstack to an array of strings. Where each string is the index of the FName. */
+	static void DecodeToStrings( const FName& EncodedCallstack, TArray<FString>& out_DecodedCallstack )
+	{
+		EncodedCallstack.ToString().ParseIntoArray( out_DecodedCallstack, CallstackSeparator, true );
+	}
+};
+
+const TCHAR* FStatsCallstack::CallstackSeparator = TEXT( "+" );
+
+/*-----------------------------------------------------------------------------
+	Allocation info
+-----------------------------------------------------------------------------*/
+
+FAllocationInfo::FAllocationInfo( uint64 InOldPtr, uint64 InPtr, int64 InSize, const TArray<FName>& InCallstack, uint32 InSequenceTag, EMemoryOperation InOp, bool bInHasBrokenCallstack ) 
+	: OldPtr( InOldPtr )
+	, Ptr( InPtr )
+	, Size( InSize )
+	, EncodedCallstack( *FStatsCallstack::Encode( InCallstack ) )
+	, SequenceTag( InSequenceTag )
+	, Op( InOp )
+	, bHasBrokenCallstack( bInHasBrokenCallstack )
+{
+
+}
+
+FAllocationInfo::FAllocationInfo( const FAllocationInfo& Other ) 
+	: OldPtr( Other.OldPtr )
+	, Ptr( Other.Ptr )
+	, Size( Other.Size )
+	, EncodedCallstack( Other.EncodedCallstack )
+	, SequenceTag( Other.SequenceTag )
+	, Op( Other.Op )
+	, bHasBrokenCallstack( Other.bHasBrokenCallstack )
+{
+
+}
+
+/*-----------------------------------------------------------------------------
+	FNodeAllocationInfo
+-----------------------------------------------------------------------------*/
+
+void FNodeAllocationInfo::SortBySize()
+{
+	ChildNodes.ValueSort( FNodeAllocationInfoSizeGreater() );
+	for (auto& It : ChildNodes)
+	{
+		It.Value->SortBySize();
+	}
+}
+
+
+void FNodeAllocationInfo::PrepareCallstackData( const TArray<FName>& InDecodedCallstack )
+{
+	DecodedCallstack = InDecodedCallstack;
+	EncodedCallstack = *FStatsCallstack::Encode( DecodedCallstack );
+	HumanReadableCallstack = FStatsCallstack::GetHumanReadable( DecodedCallstack );
+}
+
+/*-----------------------------------------------------------------------------
+	FRawStatsMemoryProfiler
+-----------------------------------------------------------------------------*/
+
+FRawStatsMemoryProfiler::FRawStatsMemoryProfiler( const TCHAR* InFilename )
+	: FStatsReadFile( InFilename, true )
+	, NumDuplicatedMemoryOperations( 0 )
+	, NumMemoryOperations( 0 )
+	, LastSequenceTagForNamedMarker( 0 )
+{}
+
+void FRawStatsMemoryProfiler::PreProcessStats()
+{
+	Super::PreProcessStats();
+
+	// Begin marker.
+	Snapshots.Add( TPairInitializer<uint32, FName>( LastSequenceTagForNamedMarker, TEXT( "BeginSnapshot" ) ) );
+}
+
+
+void FRawStatsMemoryProfiler::PostProcessStats()
+{
+	Super::PostProcessStats();
+
+	const double StartTime = FPlatformTime::Seconds();
+
+	if (!IsProcessingStopped())
+	{
+		SortSequenceAllocations();
+
+		// End marker.
+		Snapshots.Add( TPairInitializer<uint32, FName>( TNumericLimits<uint32>::Max(), TEXT( "EndSnapshot" ) ) );
+
+		// Copy snapshots.
+		SnapshotsToBeProcessed = Snapshots;
+
+		UE_LOG( LogStats, Log, TEXT( "NumMemoryOperations:   %i" ), NumMemoryOperations );
+		UE_LOG( LogStats, Log, TEXT( "SequenceAllocationNum: %i" ), SequenceAllocationArray.Num() );
+
+		GenerateAllocationMap();
+		DumpDebugAllocations();
+	}
+
+	if (!IsProcessingStopped())
+	{
+		StageProgress.Set( 100 );
+
+		const double TotalTime = FPlatformTime::Seconds() - StartTime;
+		UE_LOG( LogStats, Log, TEXT( "Post-Processing took %.2f sec(s)" ), TotalTime );
+	}
+	else
+	{
+		UE_LOG( LogStats, Warning, TEXT( "Post-Processing stopped, abandoning" ) );
+	}
+}
+
+void FRawStatsMemoryProfiler::DumpDebugAllocations()
+{
+#if	UE_BUILD_DEBUG
+
+	// Dump problematic allocations
+	DuplicatedAllocMap.ValueSort( FAllocationInfoSizeGreater() );
+
+	uint64 TotalDuplicatedMemory = 0;
+	for (const auto& It : DuplicatedAllocMap)
+	{
+		const FAllocationInfo& Alloc = It.Value;
+		TotalDuplicatedMemory += Alloc.Size;
+	}
+
+	UE_LOG( LogStats, Warning, TEXT( "Dumping duplicated alloc map" ) );
+	UE_LOG( LogStats, Warning, TEXT( "TotalDuplicatedMemory: %llu bytes (%.2f MB)" ), TotalDuplicatedMemory, TotalDuplicatedMemory / 1024.0f / 1024.0f );
+	const float MaxPctDisplayed = 0.9f;
+	uint64 DisplayedSoFar = 0;
+	for (const auto& It : DuplicatedAllocMap)
+	{
+		const FAllocationInfo& Alloc = It.Value;
+		const FString& AllocCallstack = It.Key;
+		UE_LOG( LogStats, Log, TEXT( "%lli (%.2f MB) %s" ), Alloc.Size, Alloc.Size / 1024.0f / 1024.0f, *AllocCallstack );
+
+		DisplayedSoFar += Alloc.Size;
+
+		const float CurrentPct = (float)DisplayedSoFar / (float)TotalDuplicatedMemory;
+		if (CurrentPct > MaxPctDisplayed)
+		{
+			break;
+		}
+	}
+
+#endif // UE_BUILD_DEBUG
+}
+
+void FRawStatsMemoryProfiler::FreeDebugInformation()
+{
+	DuplicatedAllocMap.Empty();
+	ZeroAllocMap.Empty();
+}
+
+void FRawStatsMemoryProfiler::GenerateAllocationMap()
+{
+	/** Map of currently alive allocations. Ptr to AllocationInfo. */
+	TMap<uint64, FAllocationInfo> AllocationMap;
+
+	// Initialize the begin snapshot.
+	auto BeginSnapshot = SnapshotsToBeProcessed[0];
+
+	SnapshotsToBeProcessed.RemoveAt( 0 );
+	PrepareSnapshot( BeginSnapshot.Value, AllocationMap );
+	auto CurrentSnapshot = SnapshotsToBeProcessed[0];
+
+	UE_LOG( LogStats, Log, TEXT( "Generating memory operations map" ) );
+
+	const int32 NumSequenceAllocations = SequenceAllocationArray.Num();
+	const int32 OnePercent = FMath::Max( NumSequenceAllocations / 100, 1024 );
+	for (int32 AllocationIndex = 0; AllocationIndex < NumSequenceAllocations; AllocationIndex++)
+	{
+		if (AllocationIndex % OnePercent == 0)
+		{
+			UpdateGenerateMemoryMapProgress( AllocationIndex );
+			if (IsProcessingStopped())
+			{
+				break;
+			}
+		}
+
+		const FAllocationInfo& Alloc = SequenceAllocationArray[AllocationIndex];
+
+		// Check named marker/snapshots
+		if (Alloc.SequenceTag > CurrentSnapshot.Key)
+		{
+			SnapshotsToBeProcessed.RemoveAt( 0 );
+			PrepareSnapshot( CurrentSnapshot.Value, AllocationMap );
+			CurrentSnapshot = SnapshotsToBeProcessed[0];
+		}
+
+		if (Alloc.Op == EMemoryOperation::Alloc)
+		{
+			ProcessAlloc( Alloc, AllocationMap );
+		}
+		else if (Alloc.Op == EMemoryOperation::Realloc)
+		{
+			// Previous Alloc or Realloc
+			if (Alloc.OldPtr != 0)
+			{
+				ProcessFree( Alloc, AllocationMap, true );
+			}
+
+#if	UE_BUILD_DEBUG
+			if (Alloc.OldPtr == 0 && Alloc.Size == 0)
+			{
+				const FString ReallocCallstack = FStatsCallstack::GetHumanReadable( Alloc.EncodedCallstack );
+				UE_LOG( LogStats, VeryVerbose, TEXT( "ReallocZero: %s %i %i/%i [%i]" ), *ReallocCallstack, Alloc.Size, Alloc.OldPtr, Alloc.Ptr, Alloc.SequenceTag );
+			}
+#endif // UE_BUILD_DEBUG
+
+			if (Alloc.Ptr != 0)
+			{
+				ProcessAlloc( Alloc, AllocationMap );
+			}			
+		}
+		else if (Alloc.Op == EMemoryOperation::Free)
+		{
+			ProcessFree( Alloc, AllocationMap, false );
+		}
+	}
+
+	auto EndSnapshot = SnapshotsToBeProcessed[0];
+	SnapshotsToBeProcessed.RemoveAt( 0 );
+	PrepareSnapshot( EndSnapshot.Value, AllocationMap );
+
+	// We don't need the allocation map. Each snapshot has its own copy.
+	AllocationMap.Empty();
+
+	SnapshotNamesArray = SnapshotNamesSet.Array();
+
+	UE_LOG( LogStats, Verbose, TEXT( "NumDuplicatedMemoryOperations: %i" ), NumDuplicatedMemoryOperations );
+	UE_LOG( LogStats, Verbose, TEXT( "NumZeroAllocs:                 %i" ), NumZeroAllocs );
+}
+
+void FRawStatsMemoryProfiler::ProcessAlloc( const FAllocationInfo& AllocInfo, TMap<uint64, FAllocationInfo>& AllocationMap )
+{
+	if( AllocInfo.Size == 0 )
+	{
+		NumZeroAllocs++;
+		ZeroAllocMap.Add( FStatsCallstack::GetHumanReadable( AllocInfo.EncodedCallstack ), AllocInfo );
+	}
+
+	const FAllocationInfo* Found = AllocationMap.Find( AllocInfo.Ptr );
+	if (!Found)
+	{
+		AllocationMap.Add( AllocInfo.Ptr, AllocInfo );
+	}
+	else
+	{
+		NumDuplicatedMemoryOperations++;
+
+#if	UE_BUILD_DEBUG
+		const FString FoundCallstack = FStatsCallstack::GetHumanReadable( Found->EncodedCallstack );
+		const FString AllocCallstack = FStatsCallstack::GetHumanReadable( AllocInfo.EncodedCallstack );
+		UE_LOG( LogStats, VeryVerbose, TEXT( "DuplicatedAlloc" ) );
+		UE_LOG( LogStats, VeryVerbose, TEXT( "FoundCallstack: %s [%s]" ), *FoundCallstack, Found->Op==EMemoryOperation::Alloc ? TEXT("Alloc") : TEXT("Realloc") );
+		UE_LOG( LogStats, VeryVerbose, TEXT( "AllocCallstack: %s [%s]" ), *AllocCallstack, AllocInfo.Op==EMemoryOperation::Alloc ? TEXT("Alloc") : TEXT("Realloc") );
+		UE_LOG( LogStats, VeryVerbose, TEXT( "Size: %i/%i Ptr: %llu/%llu Tag: %i/%i" ), Found->Size, AllocInfo.Size, Found->Ptr, AllocInfo.Ptr, Found->SequenceTag, AllocInfo.SequenceTag );
+
+		// Store the old pointer.
+		DuplicatedAllocMap.Add( FoundCallstack, *Found );
+#endif // UE_BUILD_DEBUG
+
+		// Replace pointer.
+		AllocationMap.Add( AllocInfo.Ptr, AllocInfo );
+	}
+}
+
+void FRawStatsMemoryProfiler::ProcessFree( const FAllocationInfo& FreeInfo, TMap<uint64, FAllocationInfo>& AllocationMap, const bool bReallocFree )
+{
+	// bReallocFree is not needed here, but it's easier to read the code.
+	const uint64 PtrToBeFreed = bReallocFree ? FreeInfo.OldPtr : FreeInfo.Ptr;
+	const FAllocationInfo* Found = AllocationMap.Find( PtrToBeFreed );
+	if (Found)
+	{
+		const bool bIsValid = FreeInfo.SequenceTag > Found->SequenceTag;
+		if (!bIsValid)
+		{
+			UE_LOG( LogStats, Warning, TEXT( "InvalidFree Ptr: %llu, Seq: %i/%i" ), PtrToBeFreed, FreeInfo.SequenceTag, Found->SequenceTag );
+		}
+		AllocationMap.Remove( PtrToBeFreed );
+	}
+	else
+	{
+#if	UE_BUILD_DEBUG
+		const FString FWACallstack = FStatsCallstack::GetHumanReadable( FreeInfo.EncodedCallstack );
+		UE_LOG( LogStats, VeryVerbose, TEXT( "FreeWithoutAlloc: %s, %llu" ), *FWACallstack, PtrToBeFreed );
+#endif // UE_BUILD_DEBUG
+	}
+}
+
+void FRawStatsMemoryProfiler::UpdateGenerateMemoryMapProgress( const int32 AllocationIndex )
+{
+	const double CurrentSeconds = FPlatformTime::Seconds();
+	if (CurrentSeconds > LastUpdateTime + NumSecondsBetweenUpdates)
+	{
+		const int32 PercentagePos = int32( 100.0*AllocationIndex / SequenceAllocationArray.Num() );
+		StageProgress.Set( PercentagePos );
+		UE_LOG( LogStats, Verbose, TEXT( "Processing allocations %3i%% (%10i/%10i)" ), PercentagePos, AllocationIndex, SequenceAllocationArray.Num() );
+		LastUpdateTime = CurrentSeconds;
+	}
+	
+	// Abandon support.
+	if (bShouldStopProcessing == true)
+	{
+		SetProcessingStage( EStatsProcessingStage::SPS_Stopped );
+	}
+}
+
+void FRawStatsMemoryProfiler::ProcessSpecialMessageMarkerOperation( const FStatMessage& Message, const FStackState& StackState )
+{
+	const FName RawName = Message.NameAndInfo.GetRawName();
+	if (RawName == FStatConstants::RAW_NamedMarker)
+	{
+		const FName NamedMarker = Message.GetValue_FName();
+		Snapshots.Add( TPairInitializer<uint32, FName>( LastSequenceTagForNamedMarker, NamedMarker ) );
+	}
+}
+
+void FRawStatsMemoryProfiler::ProcessMemoryOperation( EMemoryOperation MemOp, uint64 Ptr, uint64 NewPtr, int64 Size, uint32 SequenceTag, const FStackState& StackState )
+{
+	if (MemOp == EMemoryOperation::Alloc)
+	{
+		NumMemoryOperations++;
+	
+		// Add a new allocation.
+		SequenceAllocationArray.Add(
+			FAllocationInfo(
+			0,
+			Ptr,
+			Size,
+			StackState.Stack,
+			SequenceTag,
+			EMemoryOperation::Alloc,
+			StackState.bIsBrokenCallstack
+			) );
+		LastSequenceTagForNamedMarker = SequenceTag;
+
+	}
+	else if (MemOp == EMemoryOperation::Realloc)
+	{
+		NumMemoryOperations++;
+		const uint64 OldPtr = Ptr;
+
+		// Add a new realloc.
+		SequenceAllocationArray.Add(
+			FAllocationInfo(
+			OldPtr,
+			NewPtr,
+			Size,
+			StackState.Stack,
+			SequenceTag,
+			EMemoryOperation::Realloc,
+			StackState.bIsBrokenCallstack
+			) );
+		LastSequenceTagForNamedMarker = SequenceTag;
+	}
+	else if (MemOp == EMemoryOperation::Free)
+	{
+		NumMemoryOperations++;
+
+		// Add a new free.
+		SequenceAllocationArray.Add(
+			FAllocationInfo(
+			0,
+			Ptr,
+			0,
+			StackState.Stack,
+			SequenceTag,
+			EMemoryOperation::Free,
+			StackState.bIsBrokenCallstack
+			) );
+	}
+}
+
+void FRawStatsMemoryProfiler::SortSequenceAllocations()
+{
+	FScopeLogTime SLT( TEXT( "SortSequenceAllocations" ), nullptr, FScopeLogTime::ScopeLog_Milliseconds );
+
+	// Sort all memory operation by the sequence tag, iterate through all operation and generate memory usage.
+	SequenceAllocationArray.Sort( FAllocationInfoSequenceTagLess() );
+
+	// Abandon support.
+	if (bShouldStopProcessing == true)
+	{
+		SetProcessingStage( EStatsProcessingStage::SPS_Stopped );
+	}
+}
+
+void FRawStatsMemoryProfiler::GenerateScopedTreeAllocations( const TMap<FName, FCombinedAllocationInfo>& ScopedAllocations, FNodeAllocationInfo& out_Root )
+{
+	FScopeLogTime SLT( TEXT( "GenerateScopedTreeAllocations" ), nullptr, FScopeLogTime::ScopeLog_Milliseconds );
+
+	// Decode all scoped allocations, generate tree for allocations and combine them.
+	for (const auto& It : ScopedAllocations)
+	{
+		const FName& EncodedCallstack = It.Key;
+		const FCombinedAllocationInfo& CombinedAllocation = It.Value;
+
+		// Decode callstack.
+		TArray<FName> DecodedCallstack;
+		FStatsCallstack::DecodeToNames( EncodedCallstack, DecodedCallstack );
+
+		const int32 AllocationLenght = DecodedCallstack.Num();
+		check( DecodedCallstack.Num() > 0 );
+
+		FNodeAllocationInfo* CurrentNode = &out_Root;
+		// Accumulate with thread root node.
+		CurrentNode->Accumulate( CombinedAllocation );
+
+		// Iterate through the callstack and prepare all nodes if needed, and accumulate memory.
+		TArray<FName> CurrentCallstack;
+		const int32 NumEntries = DecodedCallstack.Num();
+		for (int32 Idx1 = 0; Idx1 < NumEntries; ++Idx1)
+		{
+			const FName NodeName = DecodedCallstack[Idx1];
+			CurrentCallstack.Add( NodeName );
+
+			FNodeAllocationInfo* Node = nullptr;
+			const bool bContainsNode = CurrentNode->ChildNodes.Contains( NodeName );
+			if (!bContainsNode)
+			{
+				Node = new FNodeAllocationInfo;
+				Node->Depth = Idx1;
+				Node->PrepareCallstackData( CurrentCallstack );
+
+				CurrentNode->ChildNodes.Add( NodeName, Node );
+			}
+			else
+			{
+				Node = CurrentNode->ChildNodes.FindChecked( NodeName );
+			}
+
+			// Accumulate memory usage and num allocations for all nodes in the callstack.
+			Node->Accumulate( CombinedAllocation );
+
+			// Move to the next node.
+			Node->Parent = CurrentNode;
+			CurrentNode = Node;
+		}
+	}
+
+	out_Root.SortBySize();
+}
+
+void FRawStatsMemoryProfiler::ProcessAndDumpUObjectAllocations( const FName SnapshotName )
+{
+	if (!SnapshotsWithAllocationMap.Contains(SnapshotName))
+	{
+		UE_LOG( LogStats, Warning, TEXT( "Snapshot not found: %s" ), *SnapshotName.ToString() );
+		return;
+	}
+
+	const TMap<uint64, FAllocationInfo>& AllocationMap = SnapshotsWithAllocationMap.FindChecked( SnapshotName );
+
+	FScopeLogTime SLT( TEXT( "ProcessingUObjectAllocations" ), nullptr, FScopeLogTime::ScopeLog_Seconds );
+	UE_LOG( LogStats, Warning, TEXT( "Processing UObject allocations" ) );
+
+	const FString ReportName = FString::Printf( TEXT( "%s-Memory-UObject" ), *GetPlatformName() );
+	FDiagnosticTableViewer MemoryReport( *FDiagnosticTableViewer::GetUniqueTemporaryFilePath( *ReportName ), true );
+
+	// Write a row of headings for the table's columns.
+	MemoryReport.AddColumn( TEXT( "Size (bytes)" ) );
+	MemoryReport.AddColumn( TEXT( "Size (MB)" ) );
+	MemoryReport.AddColumn( TEXT( "Count" ) );
+	MemoryReport.AddColumn( TEXT( "UObject class" ) );
+	MemoryReport.CycleRow();
+
+	TMap<FName, FCombinedAllocationInfo> UObjectAllocations;
+
+	// To minimize number of calls to expensive DecodeCallstack.
+	TMap<FName, FName> UObjectCallstackToClassMapping;
+
+	uint64 NumAllocations = 0;
+	uint64 TotalAllocatedMemory = 0;
+	for (const auto& It : AllocationMap)
+	{
+		const FAllocationInfo& Alloc = It.Value;
+
+		FName UObjectClass = UObjectCallstackToClassMapping.FindRef( Alloc.EncodedCallstack );
+		if (UObjectClass == NAME_None)
+		{
+			TArray<FName> DecodedCallstack;
+			FStatsCallstack::DecodeToNames( Alloc.EncodedCallstack, DecodedCallstack );
+
+			for (int32 Index = DecodedCallstack.Num() - 1; Index >= 0; --Index)
+			{
+				const FName LongName = DecodedCallstack[Index];
+				const bool bValid = UObjectRawNames.Contains( LongName );
+				if (bValid)
+				{
+					const FString ObjectName = FStatNameAndInfo::GetShortNameFrom( LongName ).GetPlainNameString();
+					UObjectClass = *ObjectName.Left( ObjectName.Find( TEXT( "//" ) ) );;
+					UObjectCallstackToClassMapping.Add( Alloc.EncodedCallstack, UObjectClass );
+					break;
+				}
+			}
+		}
+
+		if (UObjectClass != NAME_None)
+		{
+			FCombinedAllocationInfo& CombinedAllocation = UObjectAllocations.FindOrAdd( UObjectClass );
+			CombinedAllocation += Alloc;
+
+			TotalAllocatedMemory += Alloc.Size;
+			NumAllocations++;
+		}
+	}
+
+	// Dump memory to the log.
+	UObjectAllocations.ValueSort( FCombinedAllocationInfoSizeGreater() );
+
+	const float MaxPctDisplayed = 0.90f;
+	int32 CurrentIndex = 0;
+	uint64 DisplayedSoFar = 0;
+	UE_LOG( LogStats, VeryVerbose, TEXT( "Index, Size (Size MB), Count, UObject class" ) );
+	for (const auto& It : UObjectAllocations)
+	{
+		const FCombinedAllocationInfo& CombinedAllocation = It.Value;
+		const FName& UObjectClass = It.Key;
+
+		UE_LOG( LogStats, VeryVerbose, TEXT( "%2i, %llu (%.2f MB), %llu, %s" ),
+				CurrentIndex,
+				CombinedAllocation.Size,
+				CombinedAllocation.Size / 1024.0f / 1024.0f,
+				CombinedAllocation.Count,
+				*UObjectClass.GetPlainNameString() );
+
+		// Dump stats
+		MemoryReport.AddColumn( TEXT( "%llu" ), CombinedAllocation.Size );
+		MemoryReport.AddColumn( TEXT( "%.2f MB" ), CombinedAllocation.Size / 1024.0f / 1024.0f );
+		MemoryReport.AddColumn( TEXT( "%llu" ), CombinedAllocation.Count );
+		MemoryReport.AddColumn( *UObjectClass.GetPlainNameString() );
+		MemoryReport.CycleRow();
+
+		CurrentIndex++;
+		DisplayedSoFar += CombinedAllocation.Size;
+
+		const float CurrentPct = (float)DisplayedSoFar / (float)TotalAllocatedMemory;
+		if (CurrentPct > MaxPctDisplayed)
+		{
+			break;
+		}
+	}
+
+	UE_LOG( LogStats, VeryVerbose, TEXT( "Allocated memory: %llu bytes (%.2f MB)" ), TotalAllocatedMemory, TotalAllocatedMemory / 1024.0f / 1024.0f );
+
+	// Add a total row.
+	MemoryReport.CycleRow();
+	MemoryReport.CycleRow();
+	MemoryReport.CycleRow();
+	MemoryReport.AddColumn( TEXT( "%llu" ), TotalAllocatedMemory );
+	MemoryReport.AddColumn( TEXT( "%.2f MB" ), TotalAllocatedMemory / 1024.0f / 1024.0f );
+	MemoryReport.AddColumn( TEXT( "%llu" ), NumAllocations );
+	MemoryReport.AddColumn( TEXT( "TOTAL" ) );
+	MemoryReport.CycleRow();
+}
+
+void FRawStatsMemoryProfiler::DumpScopedAllocations( const TCHAR* Name, const TMap<FString, FCombinedAllocationInfo>& CombinedAllocations )
+{
+	if (CombinedAllocations.Num() == 0)
+	{
+		UE_LOG( LogStats, Warning, TEXT( "No scoped allocations: %s" ), Name );
+		return;
+	}
+
+	FScopeLogTime SLT( TEXT( "ProcessingScopedAllocations" ), nullptr, FScopeLogTime::ScopeLog_Seconds );
+	UE_LOG( LogStats, Warning, TEXT( "Dumping scoped allocations: %s" ), Name );
+
+	const FString ReportName = FString::Printf( TEXT( "%s-Memory-Scoped-%s" ), *GetPlatformName(), Name );
+	FDiagnosticTableViewer MemoryReport( *FDiagnosticTableViewer::GetUniqueTemporaryFilePath( *ReportName ), true );
+
+	// Write a row of headings for the table's columns.
+	MemoryReport.AddColumn( TEXT( "Size (bytes)" ) );
+	MemoryReport.AddColumn( TEXT( "Size (MB)" ) );
+	MemoryReport.AddColumn( TEXT( "Count" ) );
+	MemoryReport.AddColumn( TEXT( "Callstack" ) );
+	MemoryReport.CycleRow();
+
+	FCombinedAllocationInfo Total;
+
+	const float MaxPctDisplayed = 0.90f;
+	int32 CurrentIndex = 0;
+	UE_LOG( LogStats, VeryVerbose, TEXT( "Index, Size (Size MB), Count, Stat desc" ) );
+	for (const auto& It : CombinedAllocations)
+	{
+		const FCombinedAllocationInfo& CombinedAllocation = It.Value;
+		//const FName& EncodedCallstack = It.Key;
+		const FString AllocCallstack = It.Key;// GetCallstack( EncodedCallstack );
+
+		UE_LOG( LogStats, VeryVerbose, TEXT( "%2i, %llu (%.2f MB), %llu, %s" ),
+				CurrentIndex,
+				CombinedAllocation.Size,
+				CombinedAllocation.Size / 1024.0f / 1024.0f,
+				CombinedAllocation.Count,
+				*AllocCallstack );
+
+		// Dump stats
+		MemoryReport.AddColumn( TEXT( "%llu" ), CombinedAllocation.Size );
+		MemoryReport.AddColumn( TEXT( "%.2f MB" ), CombinedAllocation.Size / 1024.0f / 1024.0f );
+		MemoryReport.AddColumn( TEXT( "%llu" ), CombinedAllocation.Count );
+		MemoryReport.AddColumn( *AllocCallstack );
+		MemoryReport.CycleRow();
+
+		CurrentIndex++;
+		Total += CombinedAllocation;
+	}
+
+	UE_LOG( LogStats, VeryVerbose, TEXT( "Allocated memory: %llu bytes (%.2f MB)" ), Total.Size, Total.SizeMB );
+
+	// Add a total row.
+	MemoryReport.CycleRow();
+	MemoryReport.CycleRow();
+	MemoryReport.CycleRow();
+	MemoryReport.AddColumn( TEXT( "%llu" ), Total.Size );
+	MemoryReport.AddColumn( TEXT( "%.2f MB" ), Total.SizeMB );
+	MemoryReport.AddColumn( TEXT( "%llu" ), Total.Count );
+	MemoryReport.AddColumn( TEXT( "TOTAL" ) );
+	MemoryReport.CycleRow();
+}
+
+void FRawStatsMemoryProfiler::GenerateScopedAllocations( const TMap<uint64, FAllocationInfo>& InAllocationMap, TMap<FName, FCombinedAllocationInfo>& out_CombinedAllocations, uint64& TotalAllocatedMemory, uint64& NumAllocations )
+{
+	FScopeLogTime SLT( TEXT( "GenerateScopedAllocations" ), nullptr, FScopeLogTime::ScopeLog_Milliseconds );
+
+	for (const auto& It : InAllocationMap)
+	{
+		const FAllocationInfo& Alloc = It.Value;
+		FCombinedAllocationInfo& CombinedAllocation = out_CombinedAllocations.FindOrAdd( Alloc.EncodedCallstack );
+		CombinedAllocation += Alloc;
+
+		TotalAllocatedMemory += Alloc.Size;
+		NumAllocations++;
+	}
+
+	// Sort by size.
+	out_CombinedAllocations.ValueSort( FCombinedAllocationInfoSizeGreater() );
+}
+
+void FRawStatsMemoryProfiler::PrepareSnapshot( const FName SnapshotName, const TMap<uint64, FAllocationInfo>& InAllocationMap )
+{
+	FScopeLogTime SLT( TEXT( "PrepareSnapshot" ), nullptr, FScopeLogTime::ScopeLog_Milliseconds );
+
+	// Make sure the snapshot name is unique.
+	FName UniqueSnapshotName = SnapshotName;
+	while (SnapshotNamesSet.Contains( UniqueSnapshotName ))
+	{
+		UniqueSnapshotName = FName( UniqueSnapshotName, UniqueSnapshotName.GetNumber() + 1 );
+	}
+	SnapshotNamesSet.Add( UniqueSnapshotName );
+
+	SnapshotsWithAllocationMap.Add( UniqueSnapshotName, InAllocationMap );
+
+	TMap<FName, FCombinedAllocationInfo> SnapshotCombinedAllocations;
+	uint64 TotalAllocatedMemory = 0;
+	uint64 NumAllocations = 0;
+	GenerateScopedAllocations( InAllocationMap, SnapshotCombinedAllocations, TotalAllocatedMemory, NumAllocations );
+	SnapshotsWithScopedAllocations.Add( UniqueSnapshotName, SnapshotCombinedAllocations );
+
+	// Decode callstacks.
+	// Replace encoded callstacks with human readable name. For easier debugging.
+	TMap<FString, FCombinedAllocationInfo> SnapshotDecodedCombinedAllocations;
+	for (auto& It : SnapshotCombinedAllocations)
+	{
+		const FString HumanReadableCallstack = FStatsCallstack::GetHumanReadable( It.Key );
+		SnapshotDecodedCombinedAllocations.Add( HumanReadableCallstack, It.Value );
+	}
+	SnapshotsWithDecodedScopedAllocations.Add( UniqueSnapshotName, SnapshotDecodedCombinedAllocations );
+
+	UE_LOG( LogStats, Warning, TEXT( "PrepareSnapshot: %s Alloc: %i Scoped: %i Total: %.2f MB" ), *UniqueSnapshotName.ToString(), InAllocationMap.Num(), SnapshotCombinedAllocations.Num(), TotalAllocatedMemory / 1024.0f / 1024.0f );
+}
+
+void FRawStatsMemoryProfiler::CompareSnapshots( const FName BeginSnaphotName, const FName EndSnaphotName, TMap<FName, FCombinedAllocationInfo>& out_Result )
+{
+	FScopeLogTime SLT( TEXT( "CompareSnapshots" ), nullptr, FScopeLogTime::ScopeLog_Milliseconds );
+
+	const auto BeginSnaphotPtr = SnapshotsWithScopedAllocations.Find( BeginSnaphotName );
+	const auto EndSnapshotPtr = SnapshotsWithScopedAllocations.Find( EndSnaphotName );
+	if (BeginSnaphotPtr && EndSnapshotPtr)
+	{
+		// Process data.
+		TMap<FName, FCombinedAllocationInfo> BeginSnaphot = *BeginSnaphotPtr;
+		TMap<FName, FCombinedAllocationInfo> EndSnaphot = *EndSnapshotPtr;
+		TMap<FName, FCombinedAllocationInfo> Result;
+
+		for (const auto& It : EndSnaphot)
+		{
+			const FName Callstack = It.Key;
+			const FCombinedAllocationInfo EndCombinedAlloc = It.Value;
+
+			const FCombinedAllocationInfo* BeginCombinedAllocPtr = BeginSnaphot.Find( Callstack );
+			if (BeginCombinedAllocPtr)
+			{
+				FCombinedAllocationInfo CombinedAllocation;
+				CombinedAllocation += EndCombinedAlloc;
+				CombinedAllocation -= *BeginCombinedAllocPtr;
+
+				if (CombinedAllocation.IsAlive())
+				{
+					out_Result.Add( Callstack, CombinedAllocation );
+				}
+			}
+			else
+			{
+				out_Result.Add( Callstack, EndCombinedAlloc );
+			}
+		}
+
+		// Sort by size.
+		out_Result.ValueSort( FCombinedAllocationInfoSizeGreater() );
+	}
+}
+
+void FRawStatsMemoryProfiler::CompareSnapshotsHumanReadable( const FName BeginSnaphotName, const FName EndSnaphotName, TMap<FString, FCombinedAllocationInfo>& out_Result )
+{
+	FScopeLogTime SLT( TEXT( "CompareSnapshotsHumanReadable" ), nullptr, FScopeLogTime::ScopeLog_Milliseconds );
+
+	const auto BeginSnaphotPtr = SnapshotsWithDecodedScopedAllocations.Find( BeginSnaphotName );
+	const auto EndSnapshotPtr = SnapshotsWithDecodedScopedAllocations.Find( EndSnaphotName );
+	if (BeginSnaphotPtr && EndSnapshotPtr)
+	{
+		// Process data.
+		TMap<FString, FCombinedAllocationInfo> BeginSnaphot = *BeginSnaphotPtr;
+		TMap<FString, FCombinedAllocationInfo> EndSnaphot = *EndSnapshotPtr;
+
+		for (const auto& It : EndSnaphot)
+		{
+			const FString& Callstack = It.Key;
+			const FCombinedAllocationInfo EndCombinedAlloc = It.Value;
+
+			const FCombinedAllocationInfo* BeginCombinedAllocPtr = BeginSnaphot.Find( Callstack );
+			if (BeginCombinedAllocPtr)
+			{
+				FCombinedAllocationInfo CombinedAllocation;
+				CombinedAllocation += EndCombinedAlloc;
+				CombinedAllocation -= *BeginCombinedAllocPtr;
+
+				if (CombinedAllocation.IsAlive())
+				{
+					out_Result.Add( Callstack, CombinedAllocation );
+				}
+			}
+			else
+			{
+				out_Result.Add( Callstack, EndCombinedAlloc );
+			}
+		}
+
+		// Sort by size.
+		out_Result.ValueSort( FCombinedAllocationInfoSizeGreater() );
+	}
+}
\ No newline at end of file
diff --git a/Engine/Source/Developer/Profiler/Private/ProfilerRawStatsForMemory.h b/Engine/Source/Developer/Profiler/Private/ProfilerRawStatsForMemory.h
new file mode 100644
--- /dev/null
+++ b/Engine/Source/Developer/Profiler/Private/ProfilerRawStatsForMemory.h
@@ -0,0 +1,384 @@
+// Copyright 1998-2016 Epic Games, Inc. All Rights Reserved.
+
+#pragma once
+
+#include "StatsData.h"
+#include "StatsFile.h"
+
+/*-----------------------------------------------------------------------------
+	Basic structures
+-----------------------------------------------------------------------------*/
+
+/** Simple allocation info. Assumes the sequence tag are unique and doesn't wrap. */
+struct FAllocationInfo
+{
+	/** Old pointer, for Realloc. */
+	uint64 OldPtr;
+
+	/** Pointer, for Alloc, Realloc and Free. */
+	uint64 Ptr;
+
+	/** Size of the allocation. */
+	int64 Size;
+	
+	/**
+	 * Encoded callstack as FNames indices separated with +
+	 * 123+45+56+76
+	 */
+	FName EncodedCallstack;
+	
+	/**
+	 * Index of the memory operation. 
+	 * Stats are handled with packets sent from different threads.
+	 * This is used to sort the memory operations to be in the sequence order.
+	 */
+	uint32 SequenceTag;
+	
+	/** Memory operation, can be Alloc, Free or Realloc. */
+	EMemoryOperation Op;
+	
+	/** If true, the callstack is broken. Missing Start/End scope. */
+	bool bHasBrokenCallstack;
+
+	/** Initialization constructor. */
+	FAllocationInfo(
+		uint64 InOldPtr,
+		uint64 InPtr,
+		int64 InSize,
+		const TArray<FName>& InCallstack,
+		uint32 InSequenceTag,
+		EMemoryOperation InOp,
+		bool bInHasBrokenCallstack
+		);
+
+	/** Copy constructor. */
+	FAllocationInfo( const FAllocationInfo& Other );
+};
+
+/** Contains data about combined allocations for an unique scope. */
+struct FCombinedAllocationInfo
+{
+	/** Size as MB. */
+	double SizeMB;
+
+	/** Allocations count. */
+	int64 Count;
+
+	/** Human readable callstack like 'GameThread -> Preinit -> LoadModule */
+	FString HumanReadableCallstack;
+
+	/** Encoded callstack like '45+656+6565'. */
+	FName EncodedCallstack;
+
+	/** Callstack indices, encoded callstack parsed into an array [45,656,6565] */
+	TArray<FName> DecodedCallstack;
+
+	/** Size for this node and children. */
+	int64 Size;
+
+	/** Default constructor. */
+	FCombinedAllocationInfo()
+		: SizeMB( 0 )
+		, Count( 0 )
+		, Size( 0 )
+	{}
+
+	/** Operator for adding another combined allocation. */
+	FCombinedAllocationInfo& operator+=(const FCombinedAllocationInfo& Other)
+	{
+		if ((UPTRINT*)this != (UPTRINT*)&Other)
+		{
+			AddAllocation( Other.Size, Other.Count );
+		}
+		return *this;
+	}
+
+	/** Operator for subtracting another combined allocation. */
+	FCombinedAllocationInfo& operator-=(const FCombinedAllocationInfo& Other)
+	{
+		if ((UPTRINT*)this != (UPTRINT*)&Other)
+		{
+			AddAllocation( -Other.Size, -Other.Count );
+		}
+		return *this;
+	}
+
+	/** Operator for adding a new allocation. */
+	FCombinedAllocationInfo& operator+=(const FAllocationInfo& Other)
+	{
+		if ((UPTRINT*)this != (UPTRINT*)&Other)
+		{
+			AddAllocation( Other.Size, 1 );
+		}
+		return *this;
+	}
+	
+	/**
+	 * @return true, if the combined allocation is alive, the allocated data is larger than 0
+	 */
+	bool IsAlive()
+	{
+		return Size > 0;
+	}
+
+protected:
+	/** Adds a new allocation. */
+	void AddAllocation( int64 SizeToAdd, int64 CountToAdd = 1 )
+	{
+		Size += SizeToAdd;
+		Count += CountToAdd;
+
+		const double InvMB = 1.0 / 1024.0 / 1024.0;
+		SizeMB = Size * InvMB;
+	}
+};
+
+/** Contains data about node allocations, parent and children allocations. */
+struct FNodeAllocationInfo
+{
+	/** Size as MBs for this node and children. */
+	double SizeMB;
+
+	/** Allocations count for this node and children. */
+	int64 Count;
+
+	/** Human readable callstack like 'GameThread -> Preinit -> LoadModule */
+	FString HumanReadableCallstack;
+
+	/** Encoded callstack like '45+656+6565'. */
+	FName EncodedCallstack;
+
+	/** Callstack indices, encoded callstack parsed into an array [45,656,6565] */
+	TArray<FName> DecodedCallstack;
+
+	/** Parent node. */
+	FNodeAllocationInfo* Parent;
+
+	/** Child nodes. */
+	TMap<FName, FNodeAllocationInfo*> ChildNodes;
+
+	/** Size for this node and children. */
+	int64 Size;
+
+	/** Node's depth. */
+	int32 Depth;
+
+	/** Default constructor. */
+	FNodeAllocationInfo()
+		: SizeMB( 0.0 )
+		, Count( 0 )
+		, Parent( nullptr )
+		, Size( 0 )
+		, Depth( 0 )
+	{}
+
+	/** Destructor. */
+	~FNodeAllocationInfo()
+	{
+		DeleteAllChildrenNodes();
+	}
+
+	/** Accumulates this node with the other combined allocation. */
+	void Accumulate( const FCombinedAllocationInfo& Other )
+	{
+		Size += Other.Size;
+		Count += Other.Count;
+
+		const double InvMB = 1.0 / 1024.0 / 1024.0;
+		SizeMB = Size * InvMB;
+	}
+
+	/** Recursively sorts all children by size. */
+	void SortBySize();
+
+	/** Prepares all callstack data for this node. */
+	void PrepareCallstackData( const TArray<FName>& InDecodedCallstack );
+
+protected:
+	/** Recursively deletes all children. */
+	void DeleteAllChildrenNodes()
+	{
+		for (auto& It : ChildNodes)
+		{
+			delete It.Value;
+		}
+		ChildNodes.Empty();
+	}
+};
+
+/*-----------------------------------------------------------------------------
+	FRawStatsMemoryProfiler
+-----------------------------------------------------------------------------*/
+
+/** Class for managing the raw stats in the context of memory profiling. */
+class FRawStatsMemoryProfiler : public FStatsReadFile
+{
+	friend struct FStatsReader<FRawStatsMemoryProfiler>;
+	typedef FStatsReadFile Super;
+
+protected:
+
+	/** Initialization constructor. */
+	FRawStatsMemoryProfiler( const TCHAR* InFilename );
+
+	/** Called before started processing combined history. */
+	virtual void PreProcessStats() override;
+
+	/** Called after finished processing combined history. */
+	virtual void PostProcessStats() override;
+
+// 	/** Processes special message for advancing the stats frame from the game thread. */
+// 	virtual void ProcessAdvanceFrameEventGameThreadOperation( const FStatMessage& Message, const FStackState& Stack ) override
+// 	{}
+// 
+// 	/** Processes special message for advancing the stats frame from the render thread. */
+// 	virtual void ProcessAdvanceFrameEventRenderThreadOperation( const FStatMessage& Message, const FStackState& Stack ) override
+// 	{}
+// 
+// 	/** ProcessesIndicates begin of the cycle scope. */
+// 	virtual void ProcessCycleScopeStartOperation( const FStatMessage& Message, const FStackState& Stack ) override
+// 	{}
+// 
+// 	/** Indicates end of the cycle scope. */
+// 	virtual void ProcessCycleScopeEndOperation( const FStatMessage& Message, const FStackState& Stack ) override
+// 	{}
+// 
+	/** Processes special message marker used determine that we encountered a special data in the stat file. */
+	virtual void ProcessSpecialMessageMarkerOperation( const FStatMessage& Message, const FStackState& StackState ) override;
+// 
+// 	/** Processes set operation. */
+// 	virtual void ProcessSetOperation( const FStatMessage& Message, const FStackState& Stack ) override
+// 	{}
+// 
+// 	/** Processes clear operation. */
+// 	virtual void ProcessClearOperation( const FStatMessage& Message, const FStackState& Stack ) override
+// 	{}
+// 
+// 	/** Processes add operation. */
+// 	virtual void ProcessAddOperation( const FStatMessage& Message, const FStackState& Stack ) override
+// 	{}
+// 
+// 	/** Processes subtract operation. */
+// 	virtual void ProcessSubtractOperation( const FStatMessage& Message, const FStackState& Stack ) override
+// 	{}
+
+	/** Processes memory operation. @see EMemoryOperation. */
+	virtual void ProcessMemoryOperation( EMemoryOperation MemOp, uint64 Ptr, uint64 NewPtr, int64 Size, uint32 SequenceTag, const FStackState& StackState ) override;
+
+	/**
+	 * Updates process stage progress periodically, does debug logging if enabled.
+	 */
+	void UpdateGenerateMemoryMapProgress( const int32 AllocationIndex );
+
+	/**
+	 * Sorts sequence allocations.
+	 */
+	void SortSequenceAllocations();
+
+	/**
+	 *	Generates allocation map and processes all memory snapshots .
+	 */
+	void GenerateAllocationMap();
+
+	/** Processes an alloc memory operation. */
+	void ProcessAlloc( const FAllocationInfo& AllocInfo, TMap<uint64, FAllocationInfo>& AllocationMap );
+
+	/** Processes a free memory operation. */
+	void ProcessFree( const FAllocationInfo& FreeInfo, TMap<uint64, FAllocationInfo>& AllocationMap, const bool bReallocFree );
+
+	/** Dumps debug allocations. */
+	void DumpDebugAllocations();
+
+	/** Frees memory used by the debug information. */
+	void FreeDebugInformation();
+
+public:
+	/** Generates UObject allocation statistics. */
+	void ProcessAndDumpUObjectAllocations( const FName SnapshotName );
+
+	void DumpScopedAllocations( const TCHAR* Name, const TMap<FString, FCombinedAllocationInfo>& CombinedAllocations );
+
+	/** Generates callstack based allocation map. */
+	void GenerateScopedAllocations( const TMap<uint64, FAllocationInfo>& InAllocationMap, TMap<FName, FCombinedAllocationInfo>& out_CombinedAllocations, uint64& TotalAllocatedMemory, uint64& NumAllocations );
+
+	void GenerateScopedTreeAllocations( const TMap<FName, FCombinedAllocationInfo>& ScopedAllocations, FNodeAllocationInfo& out_Root );
+
+	/** Prepares data for a snapshot. */
+	void PrepareSnapshot( const FName SnapshotName, const TMap<uint64, FAllocationInfo>& InAllocationMap );
+
+	/** Compare two snapshots and saves the result for further processing. */
+	void CompareSnapshots( const FName BeginSnaphotName, const FName EndSnaphotName, TMap<FName, FCombinedAllocationInfo>& out_Result );
+
+	/** Compare two snapshots and saves the result in human readable format, results can be saved to a file. */
+	void CompareSnapshotsHumanReadable( const FName BeginSnaphotName, const FName EndSnaphotName, TMap<FString, FCombinedAllocationInfo>& out_Result );
+
+	/**
+	 * @return a set of all snapshot names.
+	 */
+	const TArray<FName> GetSnapshotNames() const
+	{
+		return SnapshotNamesArray;
+	}
+
+	/**
+	 * @return Platform's name based on the loaded ue4statsraw file.
+	 */
+	const FString& GetPlatformName() const
+	{
+		return Header.PlatformName;
+	}
+
+	/**
+	 * @return true, if the file contains any usable data.
+	 */
+	bool HasValidData() const
+	{
+		return SnapshotNamesArray.Num() >= 2;
+	}
+
+protected:
+	/**
+	 *	 All allocations ordered by the sequence tag.
+	 *	 There is an assumption that the sequence tag will not turn-around.
+	 */
+	TArray<FAllocationInfo> SequenceAllocationArray;
+
+	/** The sequence tag mapping to the named markers. */
+	TArray<TPair<uint32, FName>> Snapshots;
+
+	/** Unique snapshot names. */
+	TSet<FName> SnapshotNamesSet;
+
+	/** Unique snapshot names, the same as SnapshotNamesSet. */
+	TArray<FName> SnapshotNamesArray;
+
+	/** The sequence tag mapping to the named markers that need to processed. */
+	TArray<TPair<uint32, FName>> SnapshotsToBeProcessed;
+
+	/** Snapshots with allocation map. */
+	TMap<FName, TMap<uint64, FAllocationInfo> > SnapshotsWithAllocationMap;
+
+	/** Snapshots with callstack based allocation map. */
+	TMap<FName, TMap<FName, FCombinedAllocationInfo> > SnapshotsWithScopedAllocations;
+
+	TMap<FName, TMap<FString, FCombinedAllocationInfo> > SnapshotsWithDecodedScopedAllocations;
+
+	/** Duplicated allocation map. Debug only. */
+	TMultiMap<FString, FAllocationInfo> DuplicatedAllocMap;
+
+	/** Zero allocation map. Debug only. */
+	TMultiMap<FString, FAllocationInfo> ZeroAllocMap;
+
+	/** Number of duplicated memory operations. */
+	int32 NumDuplicatedMemoryOperations;
+
+	/** Number of Malloc(0). */
+	int32 NumZeroAllocs = 0;
+
+	/** Number of memory operations. */
+	int32 NumMemoryOperations;
+
+	/** Last sequence tag for named marker. */
+	int32 LastSequenceTagForNamedMarker;
+};
+
diff --git a/Engine/Source/Developer/Profiler/Private/ProfilerRawStatsForThreadView.cpp b/Engine/Source/Developer/Profiler/Private/ProfilerRawStatsForThreadView.cpp
new file mode 100644
--- /dev/null
+++ b/Engine/Source/Developer/Profiler/Private/ProfilerRawStatsForThreadView.cpp
@@ -0,0 +1,559 @@
+// Copyright 1998-2015 Epic Games, Inc. All Rights Reserved.
+
+#include "ProfilerPrivatePCH.h"
+
+// Only copied from ProfilerSession, still not working.
+
+/*-----------------------------------------------------------------------------
+	FRawProfilerSession
+-----------------------------------------------------------------------------*/
+
+FRawProfilerSession::FRawProfilerSession( const FString& InRawStatsFileFileath )
+: FProfilerSession( EProfilerSessionTypes::StatsFileRaw, nullptr, FGuid::NewGuid(), InRawStatsFileFileath.Replace( *FStatConstants::StatsFileRawExtension, TEXT( "" ) ) )
+, CurrentMiniViewFrame( 0 )
+{
+	OnTick = FTickerDelegate::CreateRaw( this, &FRawProfilerSession::HandleTicker );
+}
+
+FRawProfilerSession::~FRawProfilerSession()
+{
+	FTicker::GetCoreTicker().RemoveTicker( OnTickHandle );
+}
+
+bool FRawProfilerSession::HandleTicker( float DeltaTime )
+{
+#if	0
+	StatsThreadStats;
+	Stream;
+
+	enum
+	{
+		MAX_NUM_DATA_PER_TICK = 30
+	};
+	int32 NumDataThisTick = 0;
+
+	// Add the data to the mini-view.
+	for( int32 FrameIndex = CurrentMiniViewFrame; FrameIndex < Stream.FramesInfo.Num(); ++FrameIndex )
+	{
+		const FStatsFrameInfo& StatsFrameInfo = Stream.FramesInfo[FrameIndex];
+		// Convert from cycles to ms.
+		TMap<uint32, float> ThreadMS;
+		for( auto InnerIt = StatsFrameInfo.ThreadCycles.CreateConstIterator(); InnerIt; ++InnerIt )
+		{
+			ThreadMS.Add( InnerIt.Key(), StatMetaData->ConvertCyclesToMS( InnerIt.Value() ) );
+		}
+
+		// Pass the reference to the stats' metadata.
+		// #Profiler: 2014-04-03 Figure out something better later.
+		OnAddThreadTime.ExecuteIfBound( FrameIndex, ThreadMS, StatMetaData );
+
+		//CurrentMiniViewFrame++;
+		NumDataThisTick++;
+		if( NumDataThisTick > MAX_NUM_DATA_PER_TICK )
+		{
+			break;
+		}
+	}
+#endif // 0
+
+	return true;
+}
+
+static double GetSecondsPerCycle( const FStatPacketArray& Frame )
+{
+	const FName SecondsPerCycleFName = FName(TEXT("//STATGROUP_Engine//STAT_SecondsPerCycle///Seconds$32$Per$32$Cycle///////STATCAT_Advanced////"));
+	const FName SecondsPerCycleRawName = FStatConstants::RAW_SecondsPerCycle;
+	double Result = 0;
+
+	for( int32 PacketIndex = 0; PacketIndex < Frame.Packets.Num(); PacketIndex++ )
+	{
+		const FStatPacket& Packet = *Frame.Packets[PacketIndex];
+		const FStatMessagesArray& Data = Packet.StatMessages;
+
+		if( Packet.ThreadType == EThreadType::Game )
+		{
+			for( int32 Index = 0; Index < Data.Num(); Index++ )
+			{
+				const FStatMessage& Item = Data[Index];
+				check( Item.NameAndInfo.GetFlag( EStatMetaFlags::DummyAlwaysOne ) ); 
+
+				const FName LongName = Item.NameAndInfo.GetEncodedName();
+				const FName RawName = Item.NameAndInfo.GetRawName();
+				if( LongName.IsEqual( SecondsPerCycleFName, ENameCase::IgnoreCase, false ) )
+				{
+					Result = Item.GetValue_double();
+					UE_LOG( LogStats, Log, TEXT( "STAT_SecondsPerCycle is %f [ns]" ), Result*1000*1000 );
+
+					PacketIndex = Frame.Packets.Num();
+					break;
+				}
+			}	
+		}
+			
+	}
+	return Result;
+}
+
+static int64 GetFastThreadFrameTimeInternal( const FStatPacketArray& Frame, EThreadType::Type ThreadType )
+{
+	int64 Result = 0;
+
+	for( int32 PacketIndex = 0; PacketIndex < Frame.Packets.Num(); PacketIndex++ )
+	{
+		const FStatPacket& Packet = *Frame.Packets[PacketIndex];
+
+		if( Packet.ThreadType == ThreadType )
+		{
+			const FStatMessagesArray& Data = Packet.StatMessages;
+			for (int32 Index = 0; Index < Data.Num(); Index++)
+			{
+				FStatMessage const& Item = Data[Index];
+				EStatOperation::Type Op = Item.NameAndInfo.GetField<EStatOperation>();
+				FName LongName = Item.NameAndInfo.GetRawName();
+				if (Op == EStatOperation::CycleScopeStart)
+				{
+					check(Item.NameAndInfo.GetFlag(EStatMetaFlags::IsCycle));
+					Result -= Item.GetValue_int64();
+					break;
+				}
+			}
+			for (int32 Index = Data.Num() - 1; Index >= 0; Index--)
+			{
+				FStatMessage const& Item = Data[Index];
+				EStatOperation::Type Op = Item.NameAndInfo.GetField<EStatOperation>();
+				FName LongName = Item.NameAndInfo.GetRawName();
+				if (Op == EStatOperation::CycleScopeEnd)
+				{
+
+					check(Item.NameAndInfo.GetFlag(EStatMetaFlags::IsCycle));
+					Result += Item.GetValue_int64();
+					break;
+				}
+			}
+		}
+	}
+	return Result;
+}
+
+void FRawProfilerSession::PrepareLoading()
+{
+	SCOPE_LOG_TIME_FUNC();
+
+	const FString Filepath = DataFilepath + FStatConstants::StatsFileRawExtension;
+	const int64 Size = IFileManager::Get().FileSize( *Filepath );
+	if( Size < 4 )
+	{
+		UE_LOG( LogStats, Error, TEXT( "Could not open: %s" ), *Filepath );
+		return;
+	}
+	TAutoPtr<FArchive> FileReader( IFileManager::Get().CreateFileReader( *Filepath ) );
+	if( !FileReader )
+	{
+		UE_LOG( LogStats, Error, TEXT( "Could not open: %s" ), *Filepath );
+		return;
+	}
+
+	if( !Stream.ReadHeader( *FileReader ) )
+	{
+		UE_LOG( LogStats, Error, TEXT( "Could not open, bad magic: %s" ), *Filepath );
+		return;
+	}
+
+	const bool bIsFinalized = Stream.Header.IsFinalized();
+	check( bIsFinalized );
+	check( Stream.Header.Version == EStatMagicWithHeader::VERSION_5 );
+
+	TArray<FStatMessage> Messages;
+	if( Stream.Header.bRawStatsFile )
+	{
+		// Read metadata.
+		TArray<FStatMessage> MetadataMessages;
+		Stream.ReadFNamesAndMetadataMessages( *FileReader, MetadataMessages );
+		StatsThreadStats.ProcessMetaDataOnly( MetadataMessages );
+
+		const FName F00245 = FName(245, 245, 0);
+		
+		const FName F11602 = FName(11602, 11602, 0);
+		const FName F06394 = FName(6394, 6394, 0);
+
+		const int64 CurrentFilePos = FileReader->Tell();
+
+		// Update profiler's metadata.
+		StatMetaData->UpdateFromStatsState( StatsThreadStats );
+		const uint32 GameThreadID = GetMetaData()->GetGameThreadID();
+
+		// Read frames offsets.
+		Stream.ReadFramesOffsets( *FileReader );
+
+		// Buffer used to store the compressed and decompressed data.
+		TArray<uint8> SrcArray;
+		TArray<uint8> DestArray;
+		const bool bHasCompressedData = Stream.Header.HasCompressedData();
+		check(bHasCompressedData);
+
+		TMap<int64, FStatPacketArray> CombinedHistory;
+		int64 TotalPacketSize = 0;
+		int64 MaximumPacketSize = 0;
+		// Read all packets sequentially, force by the memory profiler which is now a part of the raw stats.
+		// !!CAUTION!! Frame number in the raw stats is pointless, because it is time based, not frame based.
+		// Background threads usually execute time consuming operations, so the frame number won't be valid.
+		// Needs to be combined by the thread and the time, not by the frame number.
+		{
+			int64 FrameOffset0 = Stream.FramesInfo[0].FrameFileOffset;
+			FileReader->Seek( FrameOffset0 );
+
+			const int64 FileSize = FileReader->TotalSize();
+
+			while( FileReader->Tell() < FileSize )
+			{
+				// Read the compressed data.
+				FCompressedStatsData UncompressedData( SrcArray, DestArray );
+				*FileReader << UncompressedData;
+				if( UncompressedData.HasReachedEndOfCompressedData() )
+				{
+					break;
+				}
+
+				FMemoryReader MemoryReader( DestArray, true );
+
+				FStatPacket* StatPacket = new FStatPacket();
+				Stream.ReadStatPacket( MemoryReader, *StatPacket );
+				
+				const int64 FrameNum = StatPacket->Frame;
+				FStatPacketArray& Frame = CombinedHistory.FindOrAdd(FrameNum);
+			
+				// Check if we need to combine packets from the same thread.
+				FStatPacket** CombinedPacket = Frame.Packets.FindByPredicate([&](FStatPacket* Item) -> bool
+				{
+					return Item->ThreadId == StatPacket->ThreadId;
+				});
+				
+				if( CombinedPacket )
+				{
+					(*CombinedPacket)->StatMessages += StatPacket->StatMessages;
+				}
+				else
+				{
+					Frame.Packets.Add(StatPacket);
+				}
+
+				const int64 CurrentPos = FileReader->Tell();
+				const int32 PctPos = int32(100.0f*CurrentPos/FileSize);
+
+				UE_LOG( LogStats, Log, TEXT( "%3i Processing FStatPacket: Frame %5i for thread %5i with %6i messages (%.1f MB)" ), 
+					PctPos, 
+					StatPacket->Frame, 
+					StatPacket->ThreadId, 
+					StatPacket->StatMessages.Num(), 
+					StatPacket->StatMessages.GetAllocatedSize()/1024.0f/1024.0f );
+
+				const int64 PacketSize = StatPacket->StatMessages.GetAllocatedSize();
+				TotalPacketSize += PacketSize;
+				MaximumPacketSize = FMath::Max( MaximumPacketSize, PacketSize );
+			}
+		}
+
+		UE_LOG( LogStats, Log, TEXT( "TotalPacketSize: %.1f MB, Max: %1f MB" ), 
+			TotalPacketSize/1024.0f/1024.0f, 
+			MaximumPacketSize/1024.0f/1024.0f );
+
+		TArray<int64> Frames;
+		CombinedHistory.GenerateKeyArray(Frames);
+		Frames.Sort();
+		const int64 MiddleFrame = Frames[Frames.Num()/2];
+
+
+		// Remove all frames without the game thread messages.
+		for (int32 FrameIndex = 0; FrameIndex < Frames.Num(); ++FrameIndex)
+		{
+			const int64 TargetFrame = Frames[FrameIndex];
+			const FStatPacketArray& Frame = CombinedHistory.FindChecked( TargetFrame );
+
+			const double GameThreadTimeMS = GetMetaData()->ConvertCyclesToMS( GetFastThreadFrameTimeInternal( Frame, EThreadType::Game ) );
+
+			if (GameThreadTimeMS == 0.0f)
+			{
+				CombinedHistory.Remove( TargetFrame );
+				Frames.RemoveAt( FrameIndex );
+				FrameIndex--;
+			}
+		}
+		
+	
+		StatMetaData->SecondsPerCycle = GetSecondsPerCycle( CombinedHistory.FindChecked(MiddleFrame) );
+		check( StatMetaData->GetSecondsPerCycle() > 0.0 );
+
+		//const int32 FirstGameThreadFrame = FindFirstFrameWithGameThread( CombinedHistory, Frames );
+
+		// Prepare profiler frame.
+		{
+			SCOPE_LOG_TIME( TEXT( "Preparing profiler frames" ), nullptr );
+
+			// Prepare profiler frames.
+			double ElapsedTimeMS = 0;
+
+			for( int32 FrameIndex = 0; FrameIndex < Frames.Num(); ++FrameIndex )
+			{
+				const int64 TargetFrame = Frames[FrameIndex];
+				const FStatPacketArray& Frame = CombinedHistory.FindChecked(TargetFrame);
+
+				const double GameThreadTimeMS = GetMetaData()->ConvertCyclesToMS( GetFastThreadFrameTimeInternal(Frame,EThreadType::Game) );
+
+				if( GameThreadTimeMS == 0.0f )
+				{
+					continue;
+				}
+
+				const double RenderThreadTimeMS = GetMetaData()->ConvertCyclesToMS( GetFastThreadFrameTimeInternal(Frame,EThreadType::Renderer) );
+
+				// Update mini-view, convert from cycles to ms.
+				TMap<uint32, float> ThreadTimesMS;
+				ThreadTimesMS.Add( GameThreadID, GameThreadTimeMS );
+				ThreadTimesMS.Add( GetMetaData()->GetRenderThreadID()[0], RenderThreadTimeMS );
+
+				// Pass the reference to the stats' metadata.
+				OnAddThreadTime.ExecuteIfBound( FrameIndex, ThreadTimesMS, StatMetaData );
+
+				// Create a new profiler frame and add it to the stream.
+				ElapsedTimeMS += GameThreadTimeMS;
+				FProfilerFrame* ProfilerFrame = new FProfilerFrame( TargetFrame, GameThreadTimeMS, ElapsedTimeMS );
+				ProfilerFrame->ThreadTimesMS = ThreadTimesMS;
+				ProfilerStream.AddProfilerFrame( TargetFrame, ProfilerFrame );
+			}
+		}
+	
+		// Process the raw stats data.
+		{
+			SCOPE_LOG_TIME( TEXT( "Processing the raw stats" ), nullptr );
+
+			double CycleCounterAdjustmentMS = 0.0f;
+
+			// Read the raw stats messages.
+			for( int32 FrameIndex = 0; FrameIndex < Frames.Num()-1; ++FrameIndex )
+			{
+				const int64 TargetFrame = Frames[FrameIndex];
+				const FStatPacketArray& Frame = CombinedHistory.FindChecked(TargetFrame);
+
+				FProfilerFrame* ProfilerFrame = ProfilerStream.GetProfilerFrame( FrameIndex );
+
+				UE_CLOG( FrameIndex % 8 == 0, LogStats, Log, TEXT( "Processing raw stats frame: %4i/%4i" ), FrameIndex, Frames.Num() );
+
+				ProcessStatPacketArray( Frame, *ProfilerFrame, FrameIndex ); // or ProfilerFrame->TargetFrame
+
+				// Find the first cycle counter for the game thread.
+				if( CycleCounterAdjustmentMS == 0.0f )
+				{
+					CycleCounterAdjustmentMS = ProfilerFrame->Root->CycleCounterStartTimeMS;
+				}
+
+				// Update thread time and mark profiler frame as valid and ready for use.
+				ProfilerFrame->MarkAsValid();
+			}
+
+			// Adjust all profiler frames.
+			ProfilerStream.AdjustCycleCounters( CycleCounterAdjustmentMS );
+		}
+	}
+
+	const int64 AllocatedSize = ProfilerStream.GetAllocatedSize();
+
+	// We have the whole metadata and basic information about the raw stats file, start ticking the profiler session.
+	//OnTickHandle = FTicker::GetCoreTicker().AddTicker( OnTick, 0.25f );
+
+#if	0
+	if( SessionType == EProfilerSessionTypes::OfflineRaw )
+	{
+		// Broadcast that a capture file has been fully processed.
+		OnCaptureFileProcessed.ExecuteIfBound( GetInstanceID() );
+	}
+#endif // 0
+}
+
+void FRawProfilerSession::ProcessStatPacketArray( const FStatPacketArray& StatPacketArray, FProfilerFrame& out_ProfilerFrame, int32 FrameIndex )
+{
+	// #Profiler: 2014-03-24 Standardize thread names and id
+	// #Profiler: 2014-04-22 Remove all references to the data provider, event graph etc once data graph can visualize.
+
+	// Raw stats callstack for this stat packet array.
+	TMap<FName, FProfilerStackNode*> ThreadNodes;
+
+	const FProfilerStatMetaDataRef MetaData = GetMetaData();
+
+	FProfilerSampleArray& MutableCollection = const_cast<FProfilerSampleArray&>(DataProvider->GetCollection());
+
+	// Add a root sample for this frame.
+	const uint32 FrameRootSampleIndex = DataProvider->AddHierarchicalSample( 0, MetaData->GetStatByID( 1 ).OwningGroup().ID(), 1, 0.0f, 0.0f, 1 );
+
+	// Iterate through all stats packets and raw stats messages.
+	FName GameThreadFName = NAME_None;
+	for (int32 PacketIndex = 0; PacketIndex < StatPacketArray.Packets.Num(); PacketIndex++)
+	{
+		const FStatPacket& StatPacket = *StatPacketArray.Packets[PacketIndex];
+		FName ThreadFName = StatsThreadStats.Threads.FindChecked( StatPacket.ThreadId );
+		const uint32 NewThreadID = MetaData->ThreadIDtoStatID.FindChecked( StatPacket.ThreadId );
+
+		// #Profiler: 2014-04-29 Only game or render thread is supported at this moment.
+		if (StatPacket.ThreadType != EThreadType::Game && StatPacket.ThreadType != EThreadType::Renderer)
+		{
+			continue;
+		}
+
+		// Workaround for issue with rendering thread names.
+		if (StatPacket.ThreadType == EThreadType::Renderer)
+		{
+			ThreadFName = NAME_RenderThread;
+		}
+		else if (StatPacket.ThreadType == EThreadType::Game)
+		{
+			GameThreadFName = ThreadFName;
+		}
+
+		FProfilerStackNode* ThreadNode = ThreadNodes.FindRef( ThreadFName );
+		if (!ThreadNode)
+		{
+			FString ThreadIdName = FStatsUtils::BuildUniqueThreadName( StatPacket.ThreadId );
+			FStatMessage ThreadMessage( ThreadFName, EStatDataType::ST_int64, STAT_GROUP_TO_FStatGroup( STATGROUP_Threads )::GetGroupName(), STAT_GROUP_TO_FStatGroup( STATGROUP_Threads )::GetGroupCategory(), *ThreadIdName, true, true );
+			//FStatMessage ThreadMessage( ThreadFName, EStatDataType::ST_int64, nullptr, nullptr, TEXT( "" ), true, true );
+			ThreadMessage.NameAndInfo.SetFlag( EStatMetaFlags::IsPackedCCAndDuration, true );
+			ThreadMessage.Clear();
+
+			// Add a thread sample.
+			const uint32 ThreadRootSampleIndex = DataProvider->AddHierarchicalSample
+				(
+				NewThreadID,
+				MetaData->GetStatByID( NewThreadID ).OwningGroup().ID(),
+				NewThreadID,
+				-1.0f, 1,
+				FrameRootSampleIndex
+				);
+
+			ThreadNode = ThreadNodes.Add( ThreadFName, new FProfilerStackNode( nullptr, ThreadMessage, ThreadRootSampleIndex, FrameIndex ) );
+		}
+
+
+		TArray<const FStatMessage*> StartStack;
+		TArray<FProfilerStackNode*> Stack;
+		Stack.Add( ThreadNode );
+		FProfilerStackNode* Current = Stack.Last();
+
+		const FStatMessagesArray& Data = StatPacket.StatMessages;
+		for (int32 Index = 0; Index < Data.Num(); Index++)
+		{
+			const FStatMessage& Item = Data[Index];
+
+			const EStatOperation::Type Op = Item.NameAndInfo.GetField<EStatOperation>();
+			const FName LongName = Item.NameAndInfo.GetRawName();
+			const FName ShortName = Item.NameAndInfo.GetShortName();
+
+			const FName RenderingThreadTickCommandName = TEXT( "RenderingThreadTickCommand" );
+
+			// Workaround for render thread hierarchy. EStatOperation::AdvanceFrameEventRenderThread is called within the scope.
+			if (ShortName == RenderingThreadTickCommandName)
+			{
+				continue;
+			}
+
+			if (Op == EStatOperation::CycleScopeStart || Op == EStatOperation::CycleScopeEnd || Op == EStatOperation::AdvanceFrameEventRenderThread)
+			{
+				//check( Item.NameAndInfo.GetFlag( EStatMetaFlags::IsCycle ) );
+				if (Op == EStatOperation::CycleScopeStart)
+				{
+					FProfilerStackNode* ChildNode = new FProfilerStackNode( Current, Item, -1, FrameIndex );
+					Current->Children.Add( ChildNode );
+
+					// Add a child sample.
+					const uint32 SampleIndex = DataProvider->AddHierarchicalSample
+						(
+						NewThreadID,
+						MetaData->GetStatByFName( ShortName ).OwningGroup().ID(), // GroupID
+						MetaData->GetStatByFName( ShortName ).ID(), // StatID
+						0, // DurationCycles
+						1,
+						Current->SampleIndex
+						);
+					ChildNode->SampleIndex = SampleIndex;
+
+					Stack.Add( ChildNode );
+					StartStack.Add( &Item );
+					Current = ChildNode;
+				}
+				// Workaround for render thread hierarchy. EStatOperation::AdvanceFrameEventRenderThread is called within the scope.
+				if (Op == EStatOperation::AdvanceFrameEventRenderThread)
+				{
+					int k = 0; k++;
+				}
+				if (Op == EStatOperation::CycleScopeEnd)
+				{
+					const FStatMessage ScopeStart = *StartStack.Pop();
+					const FStatMessage ScopeEnd = Item;
+					const int64 Delta = int32( uint32( ScopeEnd.GetValue_int64() ) - uint32( ScopeStart.GetValue_int64() ) );
+					Current->CyclesEnd = Current->CyclesStart + Delta;
+
+					Current->CycleCounterStartTimeMS = MetaData->ConvertCyclesToMS( Current->CyclesStart );
+					Current->CycleCounterEndTimeMS = MetaData->ConvertCyclesToMS( Current->CyclesEnd );
+
+					if (Current->CycleCounterStartTimeMS > Current->CycleCounterEndTimeMS)
+					{
+						int k = 0; k++;
+					}
+
+					check( Current->CycleCounterEndTimeMS >= Current->CycleCounterStartTimeMS );
+
+					FProfilerStackNode* ChildNode = Current;
+
+					// Update the child sample's DurationMS.
+					MutableCollection[ChildNode->SampleIndex].SetDurationCycles( Delta );
+
+					verify( Current == Stack.Pop() );
+					Current = Stack.Last();
+				}
+			}
+		}
+	}
+
+	// Calculate thread times.
+	for (auto It = ThreadNodes.CreateIterator(); It; ++It)
+	{
+		FProfilerStackNode& ThreadNode = *It.Value();
+		const int32 ChildrenNum = ThreadNode.Children.Num();
+		if (ChildrenNum > 0)
+		{
+			const int32 LastChildIndex = ThreadNode.Children.Num() - 1;
+			ThreadNode.CyclesStart = ThreadNode.Children[0]->CyclesStart;
+			ThreadNode.CyclesEnd = ThreadNode.Children[LastChildIndex]->CyclesEnd;
+			ThreadNode.CycleCounterStartTimeMS = MetaData->ConvertCyclesToMS( ThreadNode.CyclesStart );
+			ThreadNode.CycleCounterEndTimeMS = MetaData->ConvertCyclesToMS( ThreadNode.CyclesEnd );
+
+			FProfilerSample& ProfilerSample = MutableCollection[ThreadNode.SampleIndex];
+			//ProfilerSample.SetStartAndEndMS( MetaData->ConvertCyclesToMS( ThreadNode.CyclesStart ), MetaData->ConvertCyclesToMS( ThreadNode.CyclesEnd ) );
+		}
+	}
+
+	// Get the game thread time.
+	check( GameThreadFName != NAME_None );
+	const FProfilerStackNode& GameThreadNode = *ThreadNodes.FindChecked( GameThreadFName );
+	const double GameThreadStartMS = MetaData->ConvertCyclesToMS( GameThreadNode.CyclesStart );
+	const double GameThreadEndMS = MetaData->ConvertCyclesToMS( GameThreadNode.CyclesEnd );
+	//MutableCollection[FrameRootSampleIndex].SetStartAndEndMS( GameThreadStartMS, GameThreadEndMS );
+
+	// Advance frame
+	const uint32 LastFrameIndex = DataProvider->GetNumFrames();
+	DataProvider->AdvanceFrame( GameThreadEndMS - GameThreadStartMS );
+
+	// Update aggregated stats
+	//UpdateAggregatedStats( LastFrameIndex );
+
+	// Update aggregated events.
+	UpdateAggregatedEventGraphData( LastFrameIndex );
+
+	// RootNode is the same as the game thread node.
+	out_ProfilerFrame.Root->CycleCounterStartTimeMS = GameThreadStartMS;
+	out_ProfilerFrame.Root->CycleCounterEndTimeMS = GameThreadEndMS;
+
+	for (auto It = ThreadNodes.CreateIterator(); It; ++It)
+	{
+		out_ProfilerFrame.AddChild( It.Value() );
+	}
+
+	out_ProfilerFrame.SortChildren();
+}
\ No newline at end of file
diff --git a/Engine/Source/Developer/Profiler/Private/ProfilerRawStatsForThreadView.h b/Engine/Source/Developer/Profiler/Private/ProfilerRawStatsForThreadView.h
new file mode 100644
--- /dev/null
+++ b/Engine/Source/Developer/Profiler/Private/ProfilerRawStatsForThreadView.h
@@ -0,0 +1,52 @@
+// Copyright 1998-2015 Epic Games, Inc. All Rights Reserved.
+
+#pragma once
+
+
+/*-----------------------------------------------------------------------------
+	Profiler session for the raw stats files
+-----------------------------------------------------------------------------*/
+
+struct FAllocationInfo;
+
+class FRawProfilerSession : public FProfilerSession
+{
+	friend class FProfilerManager;
+	friend class FProfilerActionManager;
+
+	/** Profiler stream that contains all read raw profiler frames. */
+	FProfilerStream ProfilerStream;
+
+	/** Stats thread state, mostly used to manage the stats metadata. */
+	FStatsLoadedState StatsThreadStats;
+	FStatsReadStream Stream;
+
+	/** Index of the last processed data for the mini-view. */
+	int32 CurrentMiniViewFrame;
+
+public:
+	/**
+	* Default constructor, creates a profiler session from a capture file.
+	*/
+	FRawProfilerSession( const FString& InRawStatsFileFileath );
+
+	/** Destructor. */
+	~FRawProfilerSession();
+
+	/** Updates this profiler session. */
+	bool HandleTicker( float DeltaTime );
+
+	/** Starts a process of loading the raw stats file. */
+	void PrepareLoading();
+
+	const FProfilerStream& GetStream() const
+	{
+		return ProfilerStream;
+	}
+
+	/**
+	 *	Process all stats packets and convert them to data accessible by the profiler.
+	 *	Temporary version, will be optimized later.
+	 */
+	void ProcessStatPacketArray( const FStatPacketArray& PacketArray, FProfilerFrame& out_ProfilerFrame, int32 FrameIndex );
+};
\ No newline at end of file
diff --git a/Engine/Source/Developer/Profiler/Private/ProfilerSample.cpp b/Engine/Source/Developer/Profiler/Private/ProfilerSample.cpp
--- a/Engine/Source/Developer/Profiler/Private/ProfilerSample.cpp
+++ b/Engine/Source/Developer/Profiler/Private/ProfilerSample.cpp
@@ -24,7 +24,7 @@ FString EProfilerSampleTypes::ToDescription( Type ProfilerSampleType )
 {
 	switch( ProfilerSampleType )
 	{
-		case HierarchicalTime: return LOCTEXT("StatType_Desc_Hierarchical","Hierarchical - Displayed as a time").ToString();
+		case HierarchicalTime: return LOCTEXT("StatType_Desc_Hierarchical","Hierarchical - Displayed as a time and call count").ToString();
 		case NumberInt: return LOCTEXT("StatType_Desc_NumberInt","Numerical - Displayed as a integer number").ToString();
 		case NumberFloat: return LOCTEXT("StatType_Desc_NumberFloat","Numerical - Displayed as a floating number").ToString();
 		case Memory: return LOCTEXT("StatType_Desc_Memory","Memory - Displayed as a human readable data counter").ToString();
