Copying //UE4/Dev-Core to //UE4/Main

diff --git a/Engine/Source/Developer/Profiler/Private/ProfilerSample.h b/Engine/Source/Developer/Profiler/Private/ProfilerSample.h
--- a/Engine/Source/Developer/Profiler/Private/ProfilerSample.h
+++ b/Engine/Source/Developer/Profiler/Private/ProfilerSample.h
@@ -5,26 +5,13 @@
 /*-----------------------------------------------------------------------------
 	Type definitions
 -----------------------------------------------------------------------------*/
-#define CHUNKED_PROFILER_DATA 1
 
 /** Type definition for an array of profiler samples.*/
-#if CHUNKED_PROFILER_DATA
 typedef TChunkedArray< class FProfilerSample, 1024 * 64 > FProfilerSampleArray;
-#else
-typedef TArray< class FProfilerSample > FProfilerSampleArray;
-#endif
-
-/** Type definition for an array of pointers to a profiler sample.*/
-typedef TArray< class FProfilerSample* > FProfilerSampleArrayPtr;
 
 /** Type definition for an array of indices.*/
 typedef TArray< uint32 > FIndicesArray;
 
-/** Type definition for shared pointers to instances of FProfilerSample. */
-typedef TSharedPtr<class FProfilerSample> FProfilerSamplePtr;
-
-/** Type definition for shared references to instances of FProfilerSample. */
-typedef TSharedRef<class FProfilerSample> FProfilerSampleRef;
 
 /*-----------------------------------------------------------------------------
 	Enumerators
@@ -35,22 +22,28 @@ typedef TSharedRef<class FProfilerSample> FProfilerSampleRef;
 /** Enumerates profiler sample types. */
 struct EProfilerSampleTypes
 {
-	enum Type
+	enum Type : uint32
 	{
 		/** Hierarchical - Displayed as a time, used by hierarchical stats, as "%.3f ms". */
-		HierarchicalTime,
+		HierarchicalTime = 0,
 
 		/** Numerical - Displayed as a integer number, as "%i". */
-		NumberInt,
+		NumberInt = 1,
 
 		/** Numerical - Displayed as a floating number, as "%.2f". */
-		NumberFloat,
+		NumberFloat = 2,
 
 		/** Memory - Displayed as a human readable data counter, as "%.2f kb". */
-		Memory,
+		Memory = 3,
 
 		/** Invalid enum type, may be used as a number of enumerations. */
-		InvalidOrMax,
+		InvalidOrMax = 4,
+
+		/** For extracting the type from the combined type. */
+		Shift = 0,
+		NumBits = 3,
+		Num = (1 << NumBits) - 1,
+		Mask = Num,
 	};
 
 	/**
@@ -79,37 +72,43 @@ class FProfilerSample
 {	
 	friend class IDataProvider;
 
+	enum class EStatID : uint32
+	{
+		/** For extracting the StatID from the combined type. */
+		Shift = 13,
+		NumBits = 19,
+		Num = (1 << NumBits) - 1, // 2 ^ 19 -1 = 524287 unique stats
+		Mask = Num,
+	};
+
+	enum class EGroupID : uint32
+	{
+		/** For extracting the StatID from the combined type. */
+		Shift = 3,
+		NumBits = 10,
+		Num = (1 << NumBits) - 1, // 2^10 -1 = 1023 different groups
+		Mask = Num,
+	};
+
 protected:
 	/** Child samples of this profiler sample, as indices to profiler samples. */
 	FIndicesArray _ChildrenIndices;
 
-	/** The parent of this profiler sample, as an index to a profiler sample. */
-	uint32 _ParentIndex;	
+	/** 
+	  * Contains stat data, similar to the FStatMessage.StatData
+	  * Can be interpreted as uint64 counter or uint32 duration and uint32 callcount.
+	  */
+	int64 StatData;
 
 	/** The ID of the thread that this profiler sample was captured on. */
 	const uint32 _ThreadID;
 
-	/** The ID of the stat group this profiler sample belongs to, eg. Engine */
-	const uint32 _GroupID;
-
-	/** The ID of the stat of the profiler sample, eg. Frametime */
-	const uint32 _StatID;
-
-	/** Start time of this profiler sample, in milliseconds. */
-	const double _StartMS;
-
-	/** Duration of this profiler sample, in milliseconds. */
-	const double _DurationMS;
-
-	/**
-	 *	Counter value of this profiler sample. 
-	 *	The numerical value for NumberFloat, NumberInt or Memory. 
-	 *	The number of times this counter was called for HierarchicalTime. 
-	 */
-	const double _Counter;
-
-	/** Type of this profiler sample. */
-	TEnumAsByte<EProfilerSampleTypes::Type> _Type;	
+	/** 
+	  * The ID of the stat of the profiler sample, eg. Frametime				19bits
+	  * The ID of the stat group this profiler sample belongs to, eg. Engine	10bits
+	  * Type of this profiler sample											03bits
+	  */
+	uint32 CombinedMeta;
 
 public:
 	static SIZE_T SizeOf()
@@ -123,62 +122,55 @@ public:
 		return _ThreadID; 
 	}
 
-	/** Type of this profiler sample. */
-	const EProfilerSampleTypes::Type Type() const 
-	{ 
-		return _Type; 
-	};
-
 	/** The ID of the stat group this profiler sample belongs to, eg. Engine */
-	const uint32 GroupID() const 
+	FORCEINLINE_DEBUGGABLE const uint32 GroupID() const
 	{ 
-		return _GroupID; 
+		return (CombinedMeta >> (uint32)EGroupID::Shift) & (uint32)EGroupID::Mask;
 	}
 
 	/** The ID of the stat of the profiler sample, eg. Frametime */
-	const uint32 StatID() const 
+	FORCEINLINE_DEBUGGABLE const uint32 StatID() const
 	{ 
-		return _StatID; 
+		return (CombinedMeta >> (uint32)EStatID::Shift) & (uint32)EStatID::Mask;
 	}
 
-	/** Start time of the profiler sample, in milliseconds. */
-	const double StartMS() const 
-	{ 
-		return _StartMS; 
-	}
+	/** Type of this profiler sample. */
+	FORCEINLINE_DEBUGGABLE const EProfilerSampleTypes::Type Type() const
+	{
+		return (EProfilerSampleTypes::Type)((CombinedMeta >> (uint32)EProfilerSampleTypes::Shift) & (uint32)EProfilerSampleTypes::Mask);
+	};
 
-	/** Duration of the profiler sample, in milliseconds. */
-	const double DurationMS() const 
-	{ 
-		return _DurationMS; 
+	/** Sets new stat id and sample type. */
+	FORCEINLINE_DEBUGGABLE void SetMeta( const uint32 StatID, const uint32 GroupID, const EProfilerSampleTypes::Type SampleType )
+	{
+		checkSlow( StatID <= (uint32)EStatID::Num );
+		checkSlow( GroupID <= (uint32)EGroupID::Num );
+		checkSlow( SampleType <= (uint32)EProfilerSampleTypes::Num );
+
+		CombinedMeta = 0;
+		CombinedMeta |= (StatID & (uint32)EStatID::Mask) << (uint32)EStatID::Shift;
+		CombinedMeta |= (GroupID & (uint32)EGroupID::Mask) << (uint32)EGroupID::Shift;
+		CombinedMeta |= (SampleType & (uint32)EProfilerSampleTypes::Mask) << (uint32)EProfilerSampleTypes::Shift;
 	}
 
-	/**
-	 *	Counter value of this profiler sample. 
-	 *	The numerical value for NumberFloat, NumberInt or Memory. 
-	 *	The number of times this counter was called for Time. 
-	 */
-	const double CounterAsFloat() const 
-	{ 
-		return _Counter; 
+	/** Duration of the profiler sample, in cycles. */
+	const uint32 GetDurationCycles() const
+	{
+		return FromPackedCallCountDuration_Duration( StatData );
 	}
 
-	/**
-	 *	Counter value of this profiler sample. 
-	 *	The numerical value for NumberFloat, NumberInt or Memory. 
-	 *	The number of times this counter was called for HierarchicalTime. 
-	 */
-	const int64 CounterAsInt() const 
-	{ 
-		return (int64)_Counter;
+	/** Call count of the profiler sample, only for hierarchical. */
+	const uint32 GetCallCount() const
+	{
+		return FromPackedCallCountDuration_CallCount( StatData );
 	}
 
-	/** A parent of this profiler sample, as an index to a profiler sample. */
-	const uint32 ParentIndex() const 
-	{ 
-		return _ParentIndex; 
+	/** Value of the profiler sample, as double. */
+	const double GetDoubleValue() const
+	{
+		const double Value = *(double*)&StatData;
+		return Value;
 	}
-	 
 	/** Child samples of this profiler sample, as indices to profiler samples. */
 	const FIndicesArray& ChildrenIndices() const 
 	{ 
@@ -207,22 +199,16 @@ public:
 	 * @param InThreadID		- The ID of the thread that this profiler sample was captured on
 	 * @param InGroupID			- The ID of the stat group that this profiler sample belongs to
 	 * @param InStatID			- The ID of the stat of this profiler sample
-	 * @param InStartMs			- The start time of this profiler sample, in milliseconds
 	 * @param InDurationMs		- The duration of this profiler sample, in milliseconds
 	 * @param InCallsPerFrame	- The number of times this counter was called in the frame it was captured in
-	 * @param InParentIndex		- The parent of this profiler sample, as an index to a profiler sample
 	 *
 	 */
-	FProfilerSample( uint32 InThreadID, uint32 InGroupID, uint32 InStatID, double InStartMs, double InDurationMs, uint32 InCallsPerFrame, uint32 InParentIndex = InvalidIndex )
-		: _ParentIndex( InParentIndex )
-		, _ThreadID( InThreadID )
-		, _GroupID( InGroupID )
-		, _StatID( InStatID )
-		, _StartMS( InStartMs )
-		, _DurationMS( InDurationMs )
-		, _Counter( (double)InCallsPerFrame )
-		, _Type( EProfilerSampleTypes::HierarchicalTime )
-	{}
+	FORCEINLINE_DEBUGGABLE FProfilerSample( const uint32 InThreadID, const uint32 InGroupID, const uint32 InStatID, const uint32 InDurationCycles, const uint32 InCallsPerFrame )
+		: _ThreadID( InThreadID )
+	{
+		SetMeta( InStatID, InGroupID, EProfilerSampleTypes::HierarchicalTime );
+		StatData = ToPackedCallCountDuration( InCallsPerFrame, InDurationCycles );
+	}
 
 	/**
 	 * Initialization constructor for non-hierarchical samples.
@@ -233,75 +219,40 @@ public:
 	 * @param InProfilerSampleType	- The profiler sample type of this profiler sample
 	 *
 	 */
-	FProfilerSample( uint32 InGroupID, uint32 InStatID, double InCounter, const EProfilerSampleTypes::Type InProfilerSampleType )
-		: _ParentIndex( InvalidIndex )
-		, _ThreadID( 0 )
-		, _GroupID( InGroupID )
-		, _StatID( InStatID )
-		, _StartMS( 0.0f )
-		, _DurationMS( 0.0f )
-		, _Counter( InCounter )
-		, _Type( InProfilerSampleType )
-	{}
+	FORCEINLINE_DEBUGGABLE FProfilerSample( const uint32 InGroupID, const uint32 InStatID, const double InCounter, const EProfilerSampleTypes::Type InProfilerSampleType )
+		: _ThreadID( 0 )
+	{
+		SetMeta( InStatID, InGroupID, InProfilerSampleType );
+		*(double*)&StatData = InCounter;
+	}
 
 	/** Copy constructor. */
-	explicit FProfilerSample( const FProfilerSample& Other )
+	FORCEINLINE_DEBUGGABLE explicit FProfilerSample( const FProfilerSample& Other )
 		: _ChildrenIndices( Other._ChildrenIndices )
-		, _ParentIndex( Other._ParentIndex )
+		, StatData(Other.StatData)
 		, _ThreadID( Other._ThreadID )
-		, _GroupID( Other._GroupID )
-		, _StatID( Other._StatID )
-		, _StartMS( Other._StartMS )
-		, _DurationMS( Other._DurationMS )
-		, _Counter( Other._Counter )
-		, _Type( Other._Type )
-	{}
-
-	/** Default constructor, creates an invalid profiler sample. */
-	FProfilerSample()
-		: _ParentIndex( InvalidIndex )
-		, _ThreadID( 0 )
-		, _GroupID( InvalidIndex )
-		, _StatID( InvalidIndex )
-		, _StartMS( 0.0f )
-		, _DurationMS( 0.0f )
-		, _Counter( 0.0f )
-		, _Type( EProfilerSampleTypes::InvalidOrMax )
-	{}
-
-	/** Updates the duration of this sample, should be only used to update the root stat. */
-	void SetDurationMS( double InDurationMs )
+		, CombinedMeta( Other.CombinedMeta )		
 	{
-		double& MutableDurationRef = const_cast<double&>(_DurationMS);
-		MutableDurationRef = InDurationMs;
 	}
 
-	void SetStartAndEndMS( double InStartMS, double InEndMS )
-	{
-		double& MutableStartRef = const_cast<double&>(_StartMS);
-		MutableStartRef = InStartMS;
-
-		double& MutableDurationRef = const_cast<double&>(_DurationMS);
-		MutableDurationRef = InEndMS - InStartMS;
-	}
-
-	/** @return end time of this profiler sample, in milliseconds. */
-	const double GetEndMS() const
+	/** Default constructor, creates an invalid profiler sample. */
+	FORCEINLINE_DEBUGGABLE FProfilerSample()
+		: StatData(0)
+		, _ThreadID( 0 )		
 	{
-		return _StartMS + _DurationMS;
+		SetMeta( 0, 0, EProfilerSampleTypes::InvalidOrMax );
 	}
 
-	/** @return frame rate of this profiler sample. */
-	const float GetFramerate() const
+	/** Updates the duration of this sample, should be only used to update the root stat. */
+	void SetDurationCycles( const uint32 InDurationCycles )
 	{
-		const double Framerate = 1000.0 / _DurationMS;
-		return (float)Framerate;
+		StatData = ToPackedCallCountDuration( GetCallCount(), InDurationCycles );
 	}
 
 	/** @return true if this profiler sample is valid. */
 	const bool IsValid() const
 	{
-		return _Type != EProfilerSampleTypes::InvalidOrMax;
+		return Type() != EProfilerSampleTypes::InvalidOrMax;
 	}
 
 	/** @return true if IndexToCheck is valid. */
diff --git a/Engine/Source/Developer/Profiler/Private/ProfilerSession.cpp b/Engine/Source/Developer/Profiler/Private/ProfilerSession.cpp
--- a/Engine/Source/Developer/Profiler/Private/ProfilerSession.cpp
+++ b/Engine/Source/Developer/Profiler/Private/ProfilerSession.cpp
@@ -122,8 +122,6 @@ FGraphDataSourceRefConst FProfilerSession::CreateGraphDataSource( const uint32 I
 
 void FProfilerSession::UpdateAggregatedStats( const uint32 FrameIndex )
 {
-	// #YRX_Profiler 2015-12-08  Move to profiler client manager or ignore hier stats, and move hier stats to PopulateHierarchy_Recurrent
-
 	static FTotalTimeAndCount TimeAndCount( 0.0f, 0 );
 	PROFILER_SCOPE_LOG_TIME( TEXT( "2 FProfilerSession::UpdateAggregatedStats" ), &TimeAndCount );
 
@@ -131,11 +129,17 @@ void FProfilerSession::UpdateAggregatedStats( const uint32 FrameIndex )
 	const uint32 SampleStartIndex = IndicesForFrame.X;
 	const uint32 SampleEndIndex = IndicesForFrame.Y;
 
+	// Aggregate counters etc.
 	const FProfilerSampleArray& Collection = DataProvider->GetCollection();
-
 	for( uint32 SampleIndex = SampleStartIndex; SampleIndex < SampleEndIndex; SampleIndex++ )
 	{
 		const FProfilerSample& ProfilerSample = Collection[ SampleIndex ];
+
+		// Skip hierarchical samples to ignore misleading recursion which would be counted twice etc.
+		if (ProfilerSample.Type() == EProfilerSampleTypes::HierarchicalTime)
+		{
+			continue;
+		}
 		
 		const uint32 StatID = ProfilerSample.StatID();
 		FProfilerAggregatedStat* AggregatedStat = AggregatedStats.Find( StatID );
@@ -144,7 +148,24 @@ void FProfilerSession::UpdateAggregatedStats( const uint32 FrameIndex )
 			const FProfilerStat& ProfilerStat = GetMetaData()->GetStatByID( StatID );
 			AggregatedStat = &AggregatedStats.Add( ProfilerSample.StatID(), FProfilerAggregatedStat( ProfilerStat.Name(), ProfilerStat.OwningGroup().Name(), ProfilerSample.Type() ) );
 		}
-		(*AggregatedStat) += ProfilerSample;
+		AggregatedStat->Aggregate( ProfilerSample, StatMetaData );
+	}
+
+	// Aggregate hierarchical stat.
+	const TMap<uint32, FInclusiveTime>& InclusiveAggregates = GetInclusiveAggregateStackStats( FrameIndex );
+	for (const auto& It : InclusiveAggregates)
+	{
+		const uint32 StatID = It.Key;
+		const FInclusiveTime InclusiveTime = It.Value;
+		const FProfilerStat& ProfilerStat = GetMetaData()->GetStatByID( StatID );
+
+		FProfilerAggregatedStat* AggregatedStat = AggregatedStats.Find( StatID );
+		if (!AggregatedStat)
+		{
+			AggregatedStat = &AggregatedStats.Add( StatID, FProfilerAggregatedStat( ProfilerStat.Name(), ProfilerStat.OwningGroup().Name(), ProfilerStat.Type() ) );
+		}
+		FProfilerSample MinimalSample( 0, 0, 0, InclusiveTime.DurationCycles, InclusiveTime.CallCount );
+		AggregatedStat->Aggregate( MinimalSample, StatMetaData );
 	}
 
 	for( auto It = AggregatedStats.CreateIterator(); It; ++It )
@@ -152,39 +173,116 @@ void FProfilerSession::UpdateAggregatedStats( const uint32 FrameIndex )
 		FProfilerAggregatedStat& AggregatedStat = It.Value();
 		AggregatedStat.Advance();
 	}
-
-	// @TODO: Create map for stats TMap<uint32, TArray<uint32> >; StatID -> Sample indices for faster lookup in data providers
 }
 
-FEventGraphDataRef FProfilerSession::CreateEventGraphData( const uint32 FrameStartIndex, const uint32 FrameEndIndex, const EEventGraphTypes::Type EventGraphType )
+FEventGraphData* FProfilerSession::CombineEventGraphs( const uint32 FrameStartIndex, const uint32 FrameEndIndex )
 {
-	static FTotalTimeAndCount Current( 0.0f, 0 );
-	PROFILER_SCOPE_LOG_TIME( TEXT( "FProfilerSession::CreateEventGraphData" ), &Current );
-
-	FEventGraphData* EventGraphData = new FEventGraphData( AsShared(), FrameStartIndex );
-	for (uint32 FrameIndex = FrameStartIndex + 1; FrameIndex < FrameEndIndex + 1; ++FrameIndex)
+	FEventGraphData* EventGraphData = new FEventGraphData( this, FrameStartIndex );
+	for (uint32 FrameIndex = FrameStartIndex + 1; FrameIndex < FrameEndIndex; ++FrameIndex)
 	{
 		// Create a temporary event graph data for the specified frame.
-		const FEventGraphData CurrentEventGraphData( AsShared(), FrameIndex );
+		const FEventGraphData CurrentEventGraphData( this, FrameIndex );
 		EventGraphData->Combine( CurrentEventGraphData );
 	}
-	EventGraphData->Finalize( FrameStartIndex, FrameEndIndex + 1 );
+	return EventGraphData;
+}
 
-	if (EventGraphType == EEventGraphTypes::Average)
-	{
-		EventGraphData->SetAsAverage();
-	}
-	else if (EventGraphType == EEventGraphTypes::Maximum)
+void FProfilerSession::CombineEventGraphsTask( const uint32 FrameStartIndex, const uint32 FrameEndIndex )
+{
+	FEventGraphData* SubEventGraph = CombineEventGraphs( FrameStartIndex, FrameEndIndex );
+	CombinedSubEventGraphsLFL.Push( SubEventGraph );
+}
+
+FEventGraphContainer FProfilerSession::CreateEventGraphData( const uint32 FrameStartIndex, const uint32 FrameEndIndex )
+{
+	static FTotalTimeAndCount Current( 0.0f, 0 );
+	SCOPE_LOG_TIME_FUNC_WITH_GLOBAL( &Current );
+	
+	const uint32 TotalNumFrames = FrameEndIndex - FrameStartIndex + 1;
+	enum
 	{
+		// Minimum number of frames to combine per task
+		MIN_NUM_FRAMES_PER_TASK = 8
+	};
+
+	FEventGraphData* EventGraphData = nullptr;
 
-		EventGraphData->SetAsMaximim();
+	static const bool bUseTaskGraph = true;
+	if (!bUseTaskGraph)
+	{
+		EventGraphData = CombineEventGraphs( FrameStartIndex, FrameEndIndex );
 	}
-	else if (EventGraphType == EEventGraphTypes::Total)
+	else
 	{
-		// Nothing to do.
+		uint32 NumWorkerThreads = (uint32)FTaskGraphInterface::Get().GetNumWorkerThreads();
+		uint32 NumFramesPerTask = TotalNumFrames / (NumWorkerThreads + 1);
+
+		// Find the best configuration to utilize all worker threads.
+		while (NumFramesPerTask < MIN_NUM_FRAMES_PER_TASK)
+		{
+			NumWorkerThreads--;
+			NumFramesPerTask = TotalNumFrames / (NumWorkerThreads + 1);
+
+			if (NumWorkerThreads <= 0)
+			{
+				break;
+			}
+		}
+
+		UE_LOG( LogStats, Verbose, TEXT( "NumFrames: %u, NumWorkerThreads: %u, NumFramesPerTask: %u" ), TotalNumFrames, NumWorkerThreads, NumFramesPerTask );
+
+		uint32 NumRemainingFrames = TotalNumFrames;
+		uint32 MyFrameStartIndex = FrameStartIndex;
+		FGraphEventArray CompletionEvents;
+
+		// Don't run parallel code if not really needed.
+		if (NumFramesPerTask >= MIN_NUM_FRAMES_PER_TASK)
+		{
+			for (uint32 ThreadIndex = 0; ThreadIndex < NumWorkerThreads; ++ThreadIndex)
+			{
+				CompletionEvents.Add(FSimpleDelegateGraphTask::CreateAndDispatchWhenReady
+				(
+					FSimpleDelegateGraphTask::FDelegate::CreateRaw( this, &FProfilerSession::CombineEventGraphsTask, MyFrameStartIndex, MyFrameStartIndex + NumFramesPerTask ),
+					TStatId()
+				));
+
+				NumRemainingFrames -= NumFramesPerTask;
+				MyFrameStartIndex += NumFramesPerTask;
+			}
+		}
+
+		// Final job for remaining frames.
+		FEventGraphData* FinalSubEventGraph = CombineEventGraphs( MyFrameStartIndex, MyFrameStartIndex + NumRemainingFrames );
+		
+		// Wait for results.
+		FTaskGraphInterface::Get().WaitUntilTasksComplete( CompletionEvents );
+
+		// Combine with sub event graphs.
+		TArray<FEventGraphData*> CombinedSubEventGraphs;
+		CombinedSubEventGraphsLFL.PopAll( CombinedSubEventGraphs );
+
+		for (const auto& It : CombinedSubEventGraphs)
+		{
+			FEventGraphData * const SubEventGraph = It;
+			FinalSubEventGraph->Combine( *SubEventGraph );
+			delete SubEventGraph;
+		}
+		EventGraphData = FinalSubEventGraph;
 	}
 
-	return MakeShareable( EventGraphData );
+	check( EventGraphData );
+	EventGraphData->Finalize( FrameStartIndex, FrameEndIndex + 1 );
+
+	FEventGraphDataRef Total = MakeShareable( EventGraphData );
+
+	FEventGraphDataRef Average = Total->DuplicateAsRef();
+	Average->SetAsAverage();
+
+	FEventGraphDataRef Maximum = Total->DuplicateAsRef();
+	Maximum->SetAsMaximim();
+	
+	FEventGraphContainer EventGraphContainer( FrameStartIndex, FrameEndIndex + 1, Average, Maximum, Total );
+	return EventGraphContainer;
 }
 
 void FProfilerSession::EventGraphCombine( const FEventGraphData* Current, const uint32 InNumFrames )
@@ -224,10 +322,10 @@ void FProfilerSession::UpdateAggregatedEventGraphData( const uint32 FrameIndex )
 
 	// Create a temporary event graph data for the specified frame.
 	delete EventGraphDataCurrent;
-	EventGraphDataCurrent = new FEventGraphData( AsShared(), FrameIndex );
+	EventGraphDataCurrent = new FEventGraphData( this, FrameIndex );
 	const uint32 NumFramesLocal = SessionType == EProfilerSessionTypes::Live ? DataProvider->GetNumFrames() : 0;
 
-	static const bool bUseTaskGraph = false;
+	static const bool bUseTaskGraph = true;
 
 	if( bUseTaskGraph )
 	{
@@ -259,7 +357,10 @@ void FProfilerSession::CompletionSyncAggregatedEventGraphData()
 	}
 }
 
-
+enum
+{
+	ProfilerThreadRoot = 1
+};
 
 bool FProfilerSession::HandleTicker( float DeltaTime )
 {
@@ -285,8 +386,8 @@ bool FProfilerSession::HandleTicker( float DeltaTime )
 	// Limit processing to 50ms per frame.
 	const double TimeLimit = 50 / 1000.0;
 	double Seconds = 0;
-	
-	for( int32 FrameNumber = 0; FrameNumber < FrameToProcess.Num(); FrameNumber++ )
+
+	for( int32 It = 0; It < FrameToProcess.Num(); It++ )
 	{	
 		if( Seconds > TimeLimit )
 		{
@@ -301,10 +402,10 @@ bool FProfilerSession::HandleTicker( float DeltaTime )
 
 		FSimpleScopeSecondsCounter SecondsCounter(Seconds);
 
-		const uint32 FrameIndex = FrameToProcess[0];
+		const uint32 TargetFrame = FrameToProcess[0];
 		FrameToProcess.RemoveAt( 0 );
 
-		FProfilerDataFrame& CurrentProfilerData = FrameToProfilerDataMapping.FindChecked( FrameIndex );
+		FProfilerDataFrame& CurrentProfilerData = FrameToProfilerDataMapping.FindChecked( TargetFrame );
 
 		TMap<uint32, float> ThreadMS;
 
@@ -312,73 +413,76 @@ bool FProfilerSession::HandleTicker( float DeltaTime )
 		const TMap<uint32, FProfilerCycleGraph>& CycleGraphs = CurrentProfilerData.CycleGraphs;
 
 		// Add a root sample for this frame.
-		// HACK 2013-07-19, 13:44 STAT_Root == ThreadRoot in stats2
-		const uint32 FrameRootSampleIndex = DataProvider->AddHierarchicalSample( 0, StatMetaData->GetStatByID(1).OwningGroup().ID(), 1, 0.0f, 0.0f, 1 );
+		const uint32 FrameRootSampleIndex = DataProvider->AddHierarchicalSample( 0, StatMetaData->GetStatByID( 1 ).OwningGroup().ID(), ProfilerThreadRoot, 0, 0 );
 		
-		double GameThreadTimeMS = 0.0f;
-		double MaxThreadTimeMS = 0.0f;
+		uint32 GameThreadCycles = 0;
+		uint32 MaxThreadCycles = 0;
+
+		TMap<uint32, FInclusiveTime> StatIDToInclusiveTime;
 
-		double ThreadStartTimeMS = 0.0;
 		for( auto ThreadIt = CycleGraphs.CreateConstIterator(); ThreadIt; ++ThreadIt )
 		{
 			const uint32 ThreadID = ThreadIt.Key();
 			const FProfilerCycleGraph& ThreadGraph = ThreadIt.Value();
 
 			// Calculate total time for this thread.
-			double ThreadDurationTimeMS = 0.0;
-			ThreadStartTimeMS = CurrentProfilerData.FrameStart;
+			FInclusiveTime ThreadDuration;
+			ThreadDuration.CallCount = 1;
+
 			for( int32 Index = 0; Index < ThreadGraph.Children.Num(); Index++ )
 			{
-				ThreadDurationTimeMS += StatMetaData->ConvertCyclesToMS( ThreadGraph.Children[Index].Value );
+				ThreadDuration.DurationCycles += ThreadGraph.Children[Index].Value;
 			}
 
-			if (ThreadDurationTimeMS > 0.0)
+			if (ThreadDuration.DurationCycles > 0)
 			{
 				// Check for game thread.
 				const FString GameThreadName = FName( NAME_GameThread ).GetPlainNameString();
 				const bool bGameThreadFound = StatMetaData->GetThreadDescriptions().FindChecked( ThreadID ).Contains( GameThreadName );
 				if( bGameThreadFound )
 				{
-					GameThreadTimeMS = ThreadDurationTimeMS;
+					GameThreadCycles = ThreadDuration.DurationCycles;
 				}
 
 				// Add a root sample for each thread.
-				const uint32& NewThreadID = StatMetaData->ThreadIDtoStatID.FindChecked( ThreadID );
+				const uint32& StatThreadID = StatMetaData->ThreadIDtoStatID.FindChecked( ThreadID );
+
 				const uint32 ThreadRootSampleIndex = DataProvider->AddHierarchicalSample
 				( 
-					NewThreadID/*ThreadID*/, 
-					StatMetaData->GetStatByID(NewThreadID).OwningGroup().ID(), 
-					NewThreadID, 
-					ThreadStartTimeMS, 
-					ThreadDurationTimeMS, 
-					1, 
+					StatThreadID, 
+					StatMetaData->GetStatByID(StatThreadID).OwningGroup().ID(), 
+					StatThreadID, 
+					ThreadDuration.DurationCycles, 1, 
 					FrameRootSampleIndex 
 				);
-				ThreadMS.FindOrAdd( ThreadID ) = (float)ThreadDurationTimeMS;
+				ThreadMS.FindOrAdd( ThreadID ) = StatMetaData->ConvertCyclesToMS( ThreadDuration.DurationCycles );
 
 				// Recursively add children and parent to the root samples.
 				for( int32 Index = 0; Index < ThreadGraph.Children.Num(); Index++ )
 				{
 					const FProfilerCycleGraph& CycleGraph = ThreadGraph.Children[Index];
-					const double CycleDurationMS = StatMetaData->ConvertCyclesToMS( CycleGraph.Value );
-					const double CycleStartTimeMS = ThreadStartTimeMS;
+					const uint32 ChildDurationCycles = CycleGraph.Value;
 
-					if (CycleDurationMS > 0.0)
-					{
-						PopulateHierarchy_Recurrent( CycleGraph, CycleStartTimeMS, CycleDurationMS, ThreadRootSampleIndex );
+					if (ChildDurationCycles > 0.0)
+					{					
+						PopulateHierarchy_Recurrent( StatThreadID, CycleGraph, ChildDurationCycles, ThreadRootSampleIndex, StatIDToInclusiveTime );
 					}
-					ThreadStartTimeMS += CycleDurationMS;
 				}
 
-				MaxThreadTimeMS = FMath::Max( MaxThreadTimeMS, ThreadDurationTimeMS );
+				StatIDToInclusiveTime.Add( StatThreadID, ThreadDuration );
+				MaxThreadCycles = FMath::Max( MaxThreadCycles, ThreadDuration.DurationCycles );
 			}
 		}
 
+		InclusiveAggregateStackStats.Add( MoveTemp( StatIDToInclusiveTime ) );
+		
 		// Fix the root stat time.
 		FProfilerSampleArray& MutableCollection = const_cast<FProfilerSampleArray&>(DataProvider->GetCollection());
-		MutableCollection[FrameRootSampleIndex].SetDurationMS( GameThreadTimeMS != 0.0f ? GameThreadTimeMS : MaxThreadTimeMS );
+		MutableCollection[FrameRootSampleIndex].SetDurationCycles( GameThreadCycles != 0 ? GameThreadCycles : MaxThreadCycles );
 
-		FPSAnalyzer->AddSample( GameThreadTimeMS > 0.0f ? 1000.0f/GameThreadTimeMS : 0.0f );
+		// Update FPS analyzer.
+		const float GameThreadTimeMS = StatMetaData->ConvertCyclesToMS( GameThreadCycles );
+		FPSAnalyzer->AddSample( GameThreadTimeMS > 0.0f ? 1000.0f / GameThreadTimeMS : 0.0f );
 
 		// Process the non-hierarchical samples for the specified frame.
 		{
@@ -399,24 +503,19 @@ bool FProfilerSession::HandleTicker( float DeltaTime )
 		}
 
 		// Advance frame
-		const uint32 LastFrameIndex = DataProvider->GetNumFrames();
-		DataProvider->AdvanceFrame( MaxThreadTimeMS );
+		const uint32 DataProviderFrameIndex = DataProvider->GetNumFrames();
+		DataProvider->AdvanceFrame( StatMetaData->ConvertCyclesToMS( MaxThreadCycles ) );
 
 		// Update aggregated stats
-		UpdateAggregatedStats( LastFrameIndex );
+		UpdateAggregatedStats( DataProviderFrameIndex );
 
 		// Update aggregated events.
-		UpdateAggregatedEventGraphData( LastFrameIndex );
-
-		{
-			static FTotalTimeAndCount TotalOnAddThreadTime( 0.0f, 0 );
-			PROFILER_SCOPE_LOG_TIME( TEXT( "5     FProfilerSession::HandleTicker.OnAddThreadTime" ), &Current );
-
-			// Update mini-view.
-			OnAddThreadTime.ExecuteIfBound( LastFrameIndex, ThreadMS, StatMetaData );
-		}
+		UpdateAggregatedEventGraphData( DataProviderFrameIndex );
 
-		FrameToProfilerDataMapping.Remove( FrameIndex );
+		// Update mini-view.
+		OnAddThreadTime.ExecuteIfBound( DataProviderFrameIndex, ThreadMS, StatMetaData );
+		
+		FrameToProfilerDataMapping.Remove( TargetFrame );
 	}	
 
 	if( SessionType == EProfilerSessionTypes::StatsFile )
@@ -441,459 +540,135 @@ bool FProfilerSession::HandleTicker( float DeltaTime )
 
 void FProfilerSession::PopulateHierarchy_Recurrent
 ( 
+	const uint32 StatThreadID,
 	const FProfilerCycleGraph& ParentGraph, 
-	const double ParentStartTimeMS, 
-	const double ParentDurationMS,
-	const uint32 ParentSampleIndex 
+	const uint32 ParentDurationCycles,
+	const uint32 ParentSampleIndex,
+	TMap<uint32, FInclusiveTime>& StatIDToInclusiveTime
 )
 {
 	const FProfilerStatMetaDataRef MetaData = GetMetaData();
-	const uint32& ThreadID = MetaData->ThreadIDtoStatID.FindChecked( ParentGraph.ThreadId );
+
+#if	UE_BUILD_DEBUG
+	const FName& ParentStatName = MetaData->GetStatByID( ParentGraph.StatId ).Name();
+	const FName& ParentGroupName = MetaData->GetStatByID( ParentGraph.StatId ).OwningGroup().Name();
+#endif // UE_BUILD_DEBUG
+
+	{
+		FInclusiveTime& InclusiveTime = StatIDToInclusiveTime.FindOrAdd( ParentGraph.StatId );
+		InclusiveTime.Recursion++;
+	}
 
 	const uint32 SampleIndex = DataProvider->AddHierarchicalSample
 	( 
-		ThreadID, 
+		StatThreadID, 
 		MetaData->GetStatByID(ParentGraph.StatId).OwningGroup().ID(), 
 		ParentGraph.StatId, 
-		ParentStartTimeMS, ParentDurationMS, 
-		ParentGraph.CallsPerFrame, 
+		ParentDurationCycles, ParentGraph.CallsPerFrame, 
 		ParentSampleIndex 
 	);
 
-	double ChildStartTimeMS = ParentStartTimeMS;
-	double ChildrenDurationMS = 0.0f;
+	uint32 ChildrenDurationCycles = 0;
 
 	for( int32 DataIndex = 0; DataIndex < ParentGraph.Children.Num(); DataIndex++ )
 	{
 		const FProfilerCycleGraph& ChildCyclesCounter = ParentGraph.Children[DataIndex];
-		const double ChildDurationMS = MetaData->ConvertCyclesToMS( ChildCyclesCounter.Value );
+		const uint32 ChildDurationCycles = ChildCyclesCounter.Value;
 
-		if (ChildDurationMS > 0.0)
+		if (ChildDurationCycles > 0.0)
 		{
-			PopulateHierarchy_Recurrent( ChildCyclesCounter, ChildStartTimeMS, ChildDurationMS, SampleIndex );
+			PopulateHierarchy_Recurrent( StatThreadID, ChildCyclesCounter, ChildDurationCycles, SampleIndex, StatIDToInclusiveTime );
 		}
-		ChildStartTimeMS += ChildDurationMS;
-		ChildrenDurationMS += ChildDurationMS;
+		ChildrenDurationCycles += ChildDurationCycles;
 	}
 
-	const double SelfTimeMS = ParentDurationMS - ChildrenDurationMS;
-	if( SelfTimeMS > 0.0f && ParentGraph.Children.Num() > 0 )
+	const uint32 SelfTimeCycles = ParentDurationCycles - ChildrenDurationCycles;
+	if( SelfTimeCycles > 0.0f && ParentGraph.Children.Num() > 0 )
 	{
-		const FName& ParentStatName = MetaData->GetStatByID( ParentGraph.StatId ).Name();
-		const FName& ParentGroupName = MetaData->GetStatByID( ParentGraph.StatId ).OwningGroup().Name();
-
 		// Create a fake stat that represents this profiler sample's exclusive time.
 		// This is required if we want to create correct combined event graphs later.
 		DataProvider->AddHierarchicalSample
 		(
-			ThreadID,
+			StatThreadID,
 			MetaData->GetStatByID(0).OwningGroup().ID(),
 			0, // @see FProfilerStatMetaData.Update, 0 means "Self"
-			ChildStartTimeMS, 
-			SelfTimeMS,
-			1,
+			SelfTimeCycles,	1,
 			SampleIndex
 		);
 	}
-}
-
-void FProfilerSession::LoadComplete()
-{
-	bHasAllProfilerData = true;
-}
 
-float FProfilerSession::GetProgress() const
-{
-	if (NumFrames > 0)
 	{
-		return (float)NumFramesProcessed / NumFrames;
+		FInclusiveTime& InclusiveTime = StatIDToInclusiveTime.FindChecked( ParentGraph.StatId );
+		InclusiveTime.Recursion--;
+
+		if (InclusiveTime.Recursion == 0)
+		{
+			InclusiveTime.DurationCycles += ParentDurationCycles;
+			InclusiveTime.CallCount++;
+		}
 	}
-	return 0.0f;
 }
 
-void FProfilerSession::UpdateProfilerData( const FProfilerDataFrame& Content )
+void FProfilerSession::LoadComplete()
 {
-	FrameToProfilerDataMapping.FindOrAdd( Content.Frame ) = Content;
-	FrameToProcess.Add( Content.Frame );
+	bHasAllProfilerData = true;
 }
 
-void FProfilerSession::UpdateMetadata( const FStatMetaData& InClientStatMetaData )
+void FProfilerSession::SetNumberOfFrames( int32 InNumFrames )
 {
-	const uint32 NewStatMetaDataSize = InClientStatMetaData.GetMetaDataSize();
-	if( NewStatMetaDataSize != StatMetaDataSize )
-	{
-		ClientStatMetadata = InClientStatMetaData;
-		bRequestStatMetadataUpdate = true;
-		StatMetaDataSize = NewStatMetaDataSize;
-	}
-}
+	NumFrames = InNumFrames;
 
-/*-----------------------------------------------------------------------------
-	FRawProfilerSession
------------------------------------------------------------------------------*/
-
-FRawProfilerSession::FRawProfilerSession( const FString& InRawStatsFileFileath )
-: FProfilerSession( EProfilerSessionTypes::StatsFileRaw, nullptr, FGuid::NewGuid(), InRawStatsFileFileath.Replace( *FStatConstants::StatsFileRawExtension, TEXT( "" ) ) )
-, CurrentMiniViewFrame( 0 )
-{
-	OnTick = FTickerDelegate::CreateRaw( this, &FRawProfilerSession::HandleTicker );
-}
+	InclusiveAggregateStackStats.Reserve( NumFrames );
+	AggregatedStats.Reserve( 4096 );
 
-FRawProfilerSession::~FRawProfilerSession()
-{
-	FTicker::GetCoreTicker().RemoveTicker( OnTickHandle );
+	FrameToProfilerDataMapping.Reserve( 256 );
 }
 
-bool FRawProfilerSession::HandleTicker( float DeltaTime )
+float FProfilerSession::GetProgress() const
 {
-#if	0
-	StatsThreadStats;
-	Stream;
-
-	enum
-	{
-		MAX_NUM_DATA_PER_TICK = 30
-	};
-	int32 NumDataThisTick = 0;
-
-	// Add the data to the mini-view.
-	for( int32 FrameIndex = CurrentMiniViewFrame; FrameIndex < Stream.FramesInfo.Num(); ++FrameIndex )
+	if (NumFrames > 0)
 	{
-		const FStatsFrameInfo& StatsFrameInfo = Stream.FramesInfo[FrameIndex];
-		// Convert from cycles to ms.
-		TMap<uint32, float> ThreadMS;
-		for( auto InnerIt = StatsFrameInfo.ThreadCycles.CreateConstIterator(); InnerIt; ++InnerIt )
-		{
-			ThreadMS.Add( InnerIt.Key(), StatMetaData->ConvertCyclesToMS( InnerIt.Value() ) );
-		}
-
-		// Pass the reference to the stats' metadata.
-		// @TODO yrx 2014-04-03 Figure out something better later.
-		OnAddThreadTime.ExecuteIfBound( FrameIndex, ThreadMS, StatMetaData );
-
-		//CurrentMiniViewFrame++;
-		NumDataThisTick++;
-		if( NumDataThisTick > MAX_NUM_DATA_PER_TICK )
-		{
-			break;
-		}
+		return (float)NumFramesProcessed / NumFrames;
 	}
-#endif // 0
-
-	return true;
+	return 0.0f;
 }
 
-static double GetSecondsPerCycle( const FStatPacketArray& Frame )
+const SIZE_T FProfilerSession::GetMemoryUsage() const
 {
-	const FName SecondsPerCycleFName = FName(TEXT("//STATGROUP_Engine//STAT_SecondsPerCycle///Seconds$32$Per$32$Cycle///////STATCAT_Advanced////"));
-	const FName SecondsPerCycleRawName = FStatConstants::RAW_SecondsPerCycle;
-	double Result = 0;
+	SIZE_T MemoryUsage = 0;
+	MemoryUsage += DataProvider->GetMemoryUsage();
+	MemoryUsage += StatMetaData->GetMemoryUsage();
 
-	for( int32 PacketIndex = 0; PacketIndex < Frame.Packets.Num(); PacketIndex++ )
-	{
-		const FStatPacket& Packet = *Frame.Packets[PacketIndex];
-		const FStatMessagesArray& Data = Packet.StatMessages;
+	MemoryUsage += AggregatedStats.GetAllocatedSize();
 
-		if( Packet.ThreadType == EThreadType::Game )
-		{
-			for( int32 Index = 0; Index < Data.Num(); Index++ )
-			{
-				const FStatMessage& Item = Data[Index];
-				check( Item.NameAndInfo.GetFlag( EStatMetaFlags::DummyAlwaysOne ) ); 
+	MemoryUsage += InclusiveAggregateStackStats.GetAllocatedSize();
+	for (const auto& It : InclusiveAggregateStackStats)
+	{
+		MemoryUsage += It.GetAllocatedSize();
+	}
 
-				const FName LongName = Item.NameAndInfo.GetEncodedName();
-				const FName RawName = Item.NameAndInfo.GetRawName();
-				if( LongName.IsEqual( SecondsPerCycleFName, ENameCase::IgnoreCase, false ) )
-				{
-					Result = Item.GetValue_double();
-					UE_LOG( LogStats, Log, TEXT( "STAT_SecondsPerCycle is %f [ns]" ), Result*1000*1000 );
+	MemoryUsage += FPSAnalyzer->GetMemoryUsage();
 
-					PacketIndex = Frame.Packets.Num();
-					break;
-				}
-			}	
-		}
-			
-	}
-	return Result;
+	return MemoryUsage;
 }
 
-static int64 GetFastThreadFrameTimeInternal( const FStatPacketArray& Frame, EThreadType::Type ThreadType )
+void FProfilerSession::UpdateProfilerData( const FProfilerDataFrame& Content )
 {
-	int64 Result = 0;
-
-	for( int32 PacketIndex = 0; PacketIndex < Frame.Packets.Num(); PacketIndex++ )
-	{
-		const FStatPacket& Packet = *Frame.Packets[PacketIndex];
-
-		if( Packet.ThreadType == ThreadType )
-		{
-			const FStatMessagesArray& Data = Packet.StatMessages;
-			for (int32 Index = 0; Index < Data.Num(); Index++)
-			{
-				FStatMessage const& Item = Data[Index];
-				EStatOperation::Type Op = Item.NameAndInfo.GetField<EStatOperation>();
-				FName LongName = Item.NameAndInfo.GetRawName();
-				if (Op == EStatOperation::CycleScopeStart)
-				{
-					check(Item.NameAndInfo.GetFlag(EStatMetaFlags::IsCycle));
-					Result -= Item.GetValue_int64();
-					break;
-				}
-			}
-			for (int32 Index = Data.Num() - 1; Index >= 0; Index--)
-			{
-				FStatMessage const& Item = Data[Index];
-				EStatOperation::Type Op = Item.NameAndInfo.GetField<EStatOperation>();
-				FName LongName = Item.NameAndInfo.GetRawName();
-				if (Op == EStatOperation::CycleScopeEnd)
-				{
-
-					check(Item.NameAndInfo.GetFlag(EStatMetaFlags::IsCycle));
-					Result += Item.GetValue_int64();
-					break;
-				}
-			}
-		}
-	}
-	return Result;
+	FrameToProfilerDataMapping.FindOrAdd( Content.Frame ) = Content;
+	FrameToProcess.Add( Content.Frame );
 }
 
-void FRawProfilerSession::PrepareLoading()
+void FProfilerSession::UpdateMetadata( const FStatMetaData& InClientStatMetaData )
 {
-	SCOPE_LOG_TIME_FUNC();
-
-	const FString Filepath = DataFilepath + FStatConstants::StatsFileRawExtension;
-	const int64 Size = IFileManager::Get().FileSize( *Filepath );
-	if( Size < 4 )
-	{
-		UE_LOG( LogStats, Error, TEXT( "Could not open: %s" ), *Filepath );
-		return;
-	}
-	TAutoPtr<FArchive> FileReader( IFileManager::Get().CreateFileReader( *Filepath ) );
-	if( !FileReader )
-	{
-		UE_LOG( LogStats, Error, TEXT( "Could not open: %s" ), *Filepath );
-		return;
-	}
-
-	if( !Stream.ReadHeader( *FileReader ) )
-	{
-		UE_LOG( LogStats, Error, TEXT( "Could not open, bad magic: %s" ), *Filepath );
-		return;
-	}
-
-	const bool bIsFinalized = Stream.Header.IsFinalized();
-	check( bIsFinalized );
-	check( Stream.Header.Version == EStatMagicWithHeader::VERSION_5 );
-
-	TArray<FStatMessage> Messages;
-	if( Stream.Header.bRawStatsFile )
-	{
-		// Read metadata.
-		TArray<FStatMessage> MetadataMessages;
-		Stream.ReadFNamesAndMetadataMessages( *FileReader, MetadataMessages );
-		StatsThreadStats.ProcessMetaDataOnly( MetadataMessages );
-
-		const FName F00245 = FName(245, 245, 0);
-		
-		const FName F11602 = FName(11602, 11602, 0);
-		const FName F06394 = FName(6394, 6394, 0);
-
-		const int64 CurrentFilePos = FileReader->Tell();
-
-		// Update profiler's metadata.
-		StatMetaData->UpdateFromStatsState( StatsThreadStats );
-		const uint32 GameThreadID = GetMetaData()->GetGameThreadID();
-
-		// Read frames offsets.
-		Stream.ReadFramesOffsets( *FileReader );
-
-		// Buffer used to store the compressed and decompressed data.
-		TArray<uint8> SrcArray;
-		TArray<uint8> DestArray;
-		const bool bHasCompressedData = Stream.Header.HasCompressedData();
-		check(bHasCompressedData);
-
-		TMap<int64, FStatPacketArray> CombinedHistory;
-		int64 TotalPacketSize = 0;
-		int64 MaximumPacketSize = 0;
-		// Read all packets sequentially, force by the memory profiler which is now a part of the raw stats.
-		// !!CAUTION!! Frame number in the raw stats is pointless, because it is time based, not frame based.
-		// Background threads usually execute time consuming operations, so the frame number won't be valid.
-		// Needs to be combined by the thread and the time, not by the frame number.
-		{
-			int64 FrameOffset0 = Stream.FramesInfo[0].FrameFileOffset;
-			FileReader->Seek( FrameOffset0 );
-
-			const int64 FileSize = FileReader->TotalSize();
-
-			while( FileReader->Tell() < FileSize )
-			{
-				// Read the compressed data.
-				FCompressedStatsData UncompressedData( SrcArray, DestArray );
-				*FileReader << UncompressedData;
-				if( UncompressedData.HasReachedEndOfCompressedData() )
-				{
-					break;
-				}
-
-				FMemoryReader MemoryReader( DestArray, true );
-
-				FStatPacket* StatPacket = new FStatPacket();
-				Stream.ReadStatPacket( MemoryReader, *StatPacket );
-				
-				const int64 FrameNum = StatPacket->Frame;
-				FStatPacketArray& Frame = CombinedHistory.FindOrAdd(FrameNum);
-			
-				// Check if we need to combine packets from the same thread.
-				FStatPacket** CombinedPacket = Frame.Packets.FindByPredicate([&](FStatPacket* Item) -> bool
-				{
-					return Item->ThreadId == StatPacket->ThreadId;
-				});
-				
-				if( CombinedPacket )
-				{
-					(*CombinedPacket)->StatMessages += StatPacket->StatMessages;
-				}
-				else
-				{
-					Frame.Packets.Add(StatPacket);
-				}
-
-				const int64 CurrentPos = FileReader->Tell();
-				const int32 PctPos = int32(100.0f*CurrentPos/FileSize);
-
-				UE_LOG( LogStats, Log, TEXT( "%3i Processing FStatPacket: Frame %5i for thread %5i with %6i messages (%.1f MB)" ), 
-					PctPos, 
-					StatPacket->Frame, 
-					StatPacket->ThreadId, 
-					StatPacket->StatMessages.Num(), 
-					StatPacket->StatMessages.GetAllocatedSize()/1024.0f/1024.0f );
-
-				const int64 PacketSize = StatPacket->StatMessages.GetAllocatedSize();
-				TotalPacketSize += PacketSize;
-				MaximumPacketSize = FMath::Max( MaximumPacketSize, PacketSize );
-			}
-		}
-
-		UE_LOG( LogStats, Log, TEXT( "TotalPacketSize: %.1f MB, Max: %1f MB" ), 
-			TotalPacketSize/1024.0f/1024.0f, 
-			MaximumPacketSize/1024.0f/1024.0f );
-
-		TArray<int64> Frames;
-		CombinedHistory.GenerateKeyArray(Frames);
-		Frames.Sort();
-		const int64 MiddleFrame = Frames[Frames.Num()/2];
-
-
-		// Remove all frames without the game thread messages.
-		for (int32 FrameIndex = 0; FrameIndex < Frames.Num(); ++FrameIndex)
-		{
-			const int64 TargetFrame = Frames[FrameIndex];
-			const FStatPacketArray& Frame = CombinedHistory.FindChecked( TargetFrame );
-
-			const double GameThreadTimeMS = GetMetaData()->ConvertCyclesToMS( GetFastThreadFrameTimeInternal( Frame, EThreadType::Game ) );
-
-			if (GameThreadTimeMS == 0.0f)
-			{
-				CombinedHistory.Remove( TargetFrame );
-				Frames.RemoveAt( FrameIndex );
-				FrameIndex--;
-			}
-		}
-		
-	
-		StatMetaData->SecondsPerCycle = GetSecondsPerCycle( CombinedHistory.FindChecked(MiddleFrame) );
-		check( StatMetaData->GetSecondsPerCycle() > 0.0 );
-
-		//const int32 FirstGameThreadFrame = FindFirstFrameWithGameThread( CombinedHistory, Frames );
-
-		// Prepare profiler frame.
-		{
-			SCOPE_LOG_TIME( TEXT( "Preparing profiler frames" ), nullptr );
-
-			// Prepare profiler frames.
-			double ElapsedTimeMS = 0;
-
-			for( int32 FrameIndex = 0; FrameIndex < Frames.Num(); ++FrameIndex )
-			{
-				const int64 TargetFrame = Frames[FrameIndex];
-				const FStatPacketArray& Frame = CombinedHistory.FindChecked(TargetFrame);
-
-				const double GameThreadTimeMS = GetMetaData()->ConvertCyclesToMS( GetFastThreadFrameTimeInternal(Frame,EThreadType::Game) );
-
-				if( GameThreadTimeMS == 0.0f )
-				{
-					continue;
-				}
-
-				const double RenderThreadTimeMS = GetMetaData()->ConvertCyclesToMS( GetFastThreadFrameTimeInternal(Frame,EThreadType::Renderer) );
-
-				// Update mini-view, convert from cycles to ms.
-				TMap<uint32, float> ThreadTimesMS;
-				ThreadTimesMS.Add( GameThreadID, GameThreadTimeMS );
-				ThreadTimesMS.Add( GetMetaData()->GetRenderThreadID()[0], RenderThreadTimeMS );
-
-				// Pass the reference to the stats' metadata.
-				OnAddThreadTime.ExecuteIfBound( FrameIndex, ThreadTimesMS, StatMetaData );
-
-				// Create a new profiler frame and add it to the stream.
-				ElapsedTimeMS += GameThreadTimeMS;
-				FProfilerFrame* ProfilerFrame = new FProfilerFrame( TargetFrame, GameThreadTimeMS, ElapsedTimeMS );
-				ProfilerFrame->ThreadTimesMS = ThreadTimesMS;
-				ProfilerStream.AddProfilerFrame( TargetFrame, ProfilerFrame );
-			}
-		}
-	
-		// Process the raw stats data.
-		{
-			SCOPE_LOG_TIME( TEXT( "Processing the raw stats" ), nullptr );
-
-			double CycleCounterAdjustmentMS = 0.0f;
-
-			// Read the raw stats messages.
-			for( int32 FrameIndex = 0; FrameIndex < Frames.Num()-1; ++FrameIndex )
-			{
-				const int64 TargetFrame = Frames[FrameIndex];
-				const FStatPacketArray& Frame = CombinedHistory.FindChecked(TargetFrame);
-
-				FProfilerFrame* ProfilerFrame = ProfilerStream.GetProfilerFrame( FrameIndex );
-
-				UE_CLOG( FrameIndex % 8 == 0, LogStats, Log, TEXT( "Processing raw stats frame: %4i/%4i" ), FrameIndex, Frames.Num() );
-
-				ProcessStatPacketArray( Frame, *ProfilerFrame, FrameIndex ); // or ProfilerFrame->TargetFrame
-
-				// Find the first cycle counter for the game thread.
-				if( CycleCounterAdjustmentMS == 0.0f )
-				{
-					CycleCounterAdjustmentMS = ProfilerFrame->Root->CycleCounterStartTimeMS;
-				}
-
-				// Update thread time and mark profiler frame as valid and ready for use.
-				ProfilerFrame->MarkAsValid();
-			}
-
-			// Adjust all profiler frames.
-			ProfilerStream.AdjustCycleCounters( CycleCounterAdjustmentMS );
-		}
-	}
-
-	const int64 AllocatedSize = ProfilerStream.GetAllocatedSize();
-
-	// We have the whole metadata and basic information about the raw stats file, start ticking the profiler session.
-	//OnTickHandle = FTicker::GetCoreTicker().AddTicker( OnTick, 0.25f );
-
-#if	0
-	if( SessionType == EProfilerSessionTypes::OfflineRaw )
+	const uint32 NewStatMetaDataSize = InClientStatMetaData.GetMetaDataSize();
+	if( NewStatMetaDataSize != StatMetaDataSize )
 	{
-		// Broadcast that a capture file has been fully processed.
-		OnCaptureFileProcessed.ExecuteIfBound( GetInstanceID() );
+		ClientStatMetadata = InClientStatMetaData;
+		bRequestStatMetadataUpdate = true;
+		StatMetaDataSize = NewStatMetaDataSize;
 	}
-#endif // 0
 }
 
-
 void FProfilerStatMetaData::Update( const FStatMetaData& ClientStatMetaData )
 {
 	PROFILER_SCOPE_LOG_TIME( TEXT( "FProfilerStatMetaData.Update" ), nullptr );
@@ -906,7 +681,7 @@ void FProfilerStatMetaData::Update( const FStatMetaData& ClientStatMetaData )
 
 	InitializeGroup( NoGroupID, "NoGroup" );
 	InitializeStat( 0, NoGroupID, TEXT( "Self" ), STATTYPE_CycleCounter );
-	InitializeStat( 1, NoGroupID, FStatConstants::NAME_ThreadRoot.GetPlainNameString(), STATTYPE_CycleCounter, FStatConstants::NAME_ThreadRoot );
+	InitializeStat( ProfilerThreadRoot, NoGroupID, FStatConstants::NAME_ThreadRoot.GetPlainNameString(), STATTYPE_CycleCounter, FStatConstants::NAME_ThreadRoot );
 
 	// Iterate through all stat group descriptions.
 	for( auto It = ClientStatMetaData.GroupDescriptions.CreateConstIterator(); It; ++It )
@@ -1044,194 +819,163 @@ void FProfilerStatMetaData::UpdateFromStatsState( const FStatsThreadState& Stats
 	}
 }
 
-void FRawProfilerSession::ProcessStatPacketArray( const FStatPacketArray& StatPacketArray, FProfilerFrame& out_ProfilerFrame, int32 FrameIndex )
-{
-	// @TODO yrx 2014-03-24 Standardize thread names and id
-	// @TODO yrx 2014-04-22 Remove all references to the data provider, event graph etc once data graph can visualize.
+/*-----------------------------------------------------------------------------
+	FProfilerAggregatedStat
+-----------------------------------------------------------------------------*/
 
-	// Raw stats callstack for this stat packet array.
-	TMap<FName,FProfilerStackNode*> ThreadNodes;
+FProfilerAggregatedStat::FProfilerAggregatedStat( const FName& InStatName, const FName& InGroupName, EProfilerSampleTypes::Type InStatType ) : _StatName( InStatName )
+, _GroupName( InGroupName )
+, _ValueOneFrame( 0.0f )
+, _ValueAllFrames( 0.0f )
+, _MinValueAllFrames( FLT_MAX )
+, _MaxValueAllFrames( FLT_MIN )
 
-	const FProfilerStatMetaDataRef MetaData = GetMetaData();
-	
-	FProfilerSampleArray& MutableCollection = const_cast<FProfilerSampleArray&>(DataProvider->GetCollection());
+, _NumCallsAllFrames( 0 )
+, _NumCallsOneFrame( 0 )
+, _MinNumCallsAllFrames( MAX_uint32 )
+, _MaxNumCallsAllFrames( 0 )
+
+, _NumFrames( 0 )
+, _NumFramesWithCall( 0 )
+
+, StatType( InStatType )
+{
 
-	// Add a root sample for this frame.
-	const uint32 FrameRootSampleIndex = DataProvider->AddHierarchicalSample( 0, MetaData->GetStatByID( 1 ).OwningGroup().ID(), 1, 0.0f, 0.0f, 1 );
+}
+
+FString FProfilerAggregatedStat::ToString() const
+{
+	FString FormattedValueStr;
+	if (StatType == EProfilerSampleTypes::HierarchicalTime)
+	{
+		FormattedValueStr = FString::Printf( TEXT( "{Value Min:%.3f Avg:%.3f Max:%.3f (MS) / Calls (%.1f%%) Min:%.1f Avg:%.1f Max:%.1f}" ),
+											 MinValue(), AvgValue(), MaxValue(),
+											 FramesWithCallPct(), MinNumCalls(), AvgNumCalls(), MaxNumCalls() );
+	}
+	else if (StatType == EProfilerSampleTypes::Memory)
+	{
+		FormattedValueStr = FString::Printf( TEXT( "{Min:%.2f Avg:%.2f Max:%.2f (KB)}" ), MinValue(), AvgValue(), MaxValue() );
+	}
+	else if (StatType == EProfilerSampleTypes::NumberInt)
+	{
+		FormattedValueStr = FString::Printf( TEXT( "{Min:%.2f Avg:%.2f Max:%.2f}" ), MinValue(), AvgValue(), MaxValue() );
+	}
+	else if (StatType == EProfilerSampleTypes::NumberFloat)
+	{
+		FormattedValueStr = FString::Printf( TEXT( "{Min:%.2f Avg:%.2f Max:%.2f}" ), MinValue(), AvgValue(), MaxValue() );
+	}
+	else
+	{
+		check( 0 );
+	}
+	return FormattedValueStr;
+}
 
-	// Iterate through all stats packets and raw stats messages.
-	FName GameThreadFName = NAME_None;
-	for( int32 PacketIndex = 0; PacketIndex < StatPacketArray.Packets.Num(); PacketIndex++ )
+FString FProfilerAggregatedStat::GetFormattedValue( const Type ValueType ) const
+{
+	check( ValueType < Type::EInvalidOrMax );
+	const double ValueArray[Type::EInvalidOrMax] =
 	{
-		const FStatPacket& StatPacket = *StatPacketArray.Packets[PacketIndex];
-		FName ThreadFName = StatsThreadStats.Threads.FindChecked( StatPacket.ThreadId );
-		const uint32 NewThreadID = MetaData->ThreadIDtoStatID.FindChecked( StatPacket.ThreadId );
+		AvgValue(),
+		MinValue(),
+		MaxValue(),
 
-		// @TODO yrx 2014-04-29 Only game or render thread is supported at this moment.
-		if( StatPacket.ThreadType != EThreadType::Game && StatPacket.ThreadType != EThreadType::Renderer )
-		{
-			continue;
-		}
+		AvgNumCalls(),
+		MinNumCalls(),
+		MaxNumCalls(),
+
+		FramesWithCallPct(),
+	};
 
-		// Workaround for issue with rendering thread names.
-		if( StatPacket.ThreadType == EThreadType::Renderer )
+	FString FormatValueStr;
+	if (StatType == EProfilerSampleTypes::HierarchicalTime)
+	{
+		if (ValueType == EMinValue || ValueType == EAvgValue || ValueType == EMaxValue)
 		{
-			ThreadFName = NAME_RenderThread;
+			FormatValueStr = FString::Printf( TEXT( "%.3f (MS)" ), ValueArray[ValueType] );
 		}
-		else if( StatPacket.ThreadType == EThreadType::Game )
+		else if (ValueType == EFramesWithCallPct)
 		{
-			GameThreadFName = ThreadFName;
+			FormatValueStr = FString::Printf( TEXT( "%.1f%%" ), ValueArray[EFramesWithCallPct] );
 		}
-
-		FProfilerStackNode* ThreadNode = ThreadNodes.FindRef( ThreadFName );
-		if( !ThreadNode )
+		else
 		{
-			FString ThreadIdName = FStatsUtils::BuildUniqueThreadName( StatPacket.ThreadId );
-			FStatMessage ThreadMessage( ThreadFName, EStatDataType::ST_int64, STAT_GROUP_TO_FStatGroup( STATGROUP_Threads )::GetGroupName(), STAT_GROUP_TO_FStatGroup( STATGROUP_Threads )::GetGroupCategory(), *ThreadIdName, true, true );
-			//FStatMessage ThreadMessage( ThreadFName, EStatDataType::ST_int64, nullptr, nullptr, TEXT( "" ), true, true );
-			ThreadMessage.NameAndInfo.SetFlag( EStatMetaFlags::IsPackedCCAndDuration, true );
-			ThreadMessage.Clear();
-
-			// Add a thread sample.
-			const uint32 ThreadRootSampleIndex = DataProvider->AddHierarchicalSample
-			(
-				NewThreadID,
-				MetaData->GetStatByID( NewThreadID ).OwningGroup().ID(),
-				NewThreadID,
-				-1.0f,
-				-1.0f,
-				1,
-				FrameRootSampleIndex
-			);
-
-			ThreadNode = ThreadNodes.Add( ThreadFName, new FProfilerStackNode( nullptr, ThreadMessage, ThreadRootSampleIndex, FrameIndex ) );
+			FormatValueStr = FString::Printf( TEXT( "%.1f" ), ValueArray[ValueType] );
 		}
+	}
+	else if (StatType == EProfilerSampleTypes::Memory)
+	{
+		FormatValueStr = FString::Printf( TEXT( "%.2f (KB)" ), ValueArray[ValueType] );
+	}
+	else if (StatType == EProfilerSampleTypes::NumberInt)
+	{
+		FormatValueStr = FString::Printf( TEXT( "%.2f" ), ValueArray[ValueType] );
+	}
+	else if (StatType == EProfilerSampleTypes::NumberFloat)
+	{
+		FormatValueStr = FString::Printf( TEXT( "%.2f" ), ValueArray[ValueType] );
+	}
+	else
+	{
+		check( 0 );
+	}
+	return FormatValueStr;
+}
 
+void FProfilerAggregatedStat::Advance()
+{
+	_NumFrames++;
 
-		TArray<const FStatMessage*> StartStack;
-		TArray<FProfilerStackNode*> Stack;
-		Stack.Add( ThreadNode );
-		FProfilerStackNode* Current = Stack.Last();
+	_NumCallsAllFrames += _NumCallsOneFrame;
+	_ValueAllFrames += _ValueOneFrame;
 
-		const FStatMessagesArray& Data = StatPacket.StatMessages;
-		for( int32 Index = 0; Index < Data.Num(); Index++ )
+	// Calculate new extreme values.
+	if (StatType == EProfilerSampleTypes::HierarchicalTime)
+	{
+		// Update the extreme values only if this stat has been called at least once.
+		if (_NumCallsOneFrame > 0)
 		{
-			const FStatMessage& Item = Data[Index];
-
-			const EStatOperation::Type Op = Item.NameAndInfo.GetField<EStatOperation>();
-			const FName LongName = Item.NameAndInfo.GetRawName();
-			const FName ShortName = Item.NameAndInfo.GetShortName();
-
-			const FName RenderingThreadTickCommandName = TEXT("RenderingThreadTickCommand");
-
-			// Workaround for render thread hierarchy. EStatOperation::AdvanceFrameEventRenderThread is called within the scope.
-			if( ShortName == RenderingThreadTickCommandName )
-			{
-				continue;
-			}
-
-			if( Op == EStatOperation::CycleScopeStart || Op == EStatOperation::CycleScopeEnd || Op == EStatOperation::AdvanceFrameEventRenderThread )
-			{
-				//check( Item.NameAndInfo.GetFlag( EStatMetaFlags::IsCycle ) );
-				if( Op == EStatOperation::CycleScopeStart )
-				{
-					FProfilerStackNode* ChildNode = new FProfilerStackNode( Current, Item, -1, FrameIndex );
-					Current->Children.Add( ChildNode );
-
-					// Add a child sample.
-					const uint32 SampleIndex = DataProvider->AddHierarchicalSample
-					(
-						NewThreadID,
-						MetaData->GetStatByFName( ShortName ).OwningGroup().ID(), // GroupID
-						MetaData->GetStatByFName( ShortName ).ID(), // StatID
-						MetaData->ConvertCyclesToMS( ChildNode->CyclesStart ), // StartMS 
-						MetaData->ConvertCyclesToMS( 0 ), // DurationMS
-						1,
-						Current->SampleIndex
-					);
-					ChildNode->SampleIndex = SampleIndex;
-
-					Stack.Add( ChildNode );
-					StartStack.Add( &Item );
-					Current = ChildNode;
-				}
-				// Workaround for render thread hierarchy. EStatOperation::AdvanceFrameEventRenderThread is called within the scope.
-				if( Op == EStatOperation::AdvanceFrameEventRenderThread )
-				{
-					int k=0;k++;
-				}
-				if( Op == EStatOperation::CycleScopeEnd )
-				{
-					const FStatMessage ScopeStart = *StartStack.Pop();
-					const FStatMessage ScopeEnd = Item;
-					const int64 Delta = int32( uint32( ScopeEnd.GetValue_int64() ) - uint32( ScopeStart.GetValue_int64() ) );
-					Current->CyclesEnd = Current->CyclesStart + Delta;
-
-					Current->CycleCounterStartTimeMS = MetaData->ConvertCyclesToMS( Current->CyclesStart );
-					Current->CycleCounterEndTimeMS = MetaData->ConvertCyclesToMS( Current->CyclesEnd );
+			_NumFramesWithCall++;
+		}
 
-					if( Current->CycleCounterStartTimeMS > Current->CycleCounterEndTimeMS )
-					{
-						int k=0;k++;
-					}
+		{
+			_MinValueAllFrames = FMath::Min<float>( _MinValueAllFrames, _ValueOneFrame );
+			_MaxValueAllFrames = FMath::Max<float>( _MaxValueAllFrames, _ValueOneFrame );
 
-					check( Current->CycleCounterEndTimeMS >= Current->CycleCounterStartTimeMS );
+			_MinNumCallsAllFrames = FMath::Min<float>( _MinNumCallsAllFrames, _NumCallsOneFrame );
+			_MaxNumCallsAllFrames = FMath::Max<float>( _MaxNumCallsAllFrames, _NumCallsOneFrame );
+		}
+	}
+	else
+	{
+		_MinValueAllFrames = FMath::Min<float>( _MinValueAllFrames, _ValueOneFrame );
+		_MaxValueAllFrames = FMath::Max<float>( _MaxValueAllFrames, _ValueOneFrame );
+	}
 
-					FProfilerStackNode* ChildNode = Current;
+	_ValueOneFrame = 0.0f;
+	_NumCallsOneFrame = 0;
+}
 
-					// Update the child sample's DurationMS.
-					MutableCollection[ChildNode->SampleIndex].SetDurationMS( MetaData->ConvertCyclesToMS( Delta ) );
+void FProfilerAggregatedStat::Aggregate(const FProfilerSample& Sample, const FProfilerStatMetaDataRef& Metadata)
+{
+	double TypedValue = 0.0;
 
-					verify( Current == Stack.Pop() );
-					Current = Stack.Last();				
-				}
-			}
-		}
+	// Determine whether to we are reading a time hierarchical sample or not.
+	if (Sample.Type() == EProfilerSampleTypes::HierarchicalTime)
+	{
+		TypedValue = Metadata->ConvertCyclesToMS( Sample.GetDurationCycles() );
+		_NumCallsOneFrame += Sample.GetCallCount();
 	}
-
-	// Calculate thread times.
-	for( auto It = ThreadNodes.CreateIterator(); It; ++It )
+	else
 	{
-		FProfilerStackNode& ThreadNode = *It.Value();
-		const int32 ChildrenNum = ThreadNode.Children.Num();
-		if( ChildrenNum > 0 )
+		TypedValue = Sample.GetDoubleValue();
+
+		if (Sample.Type() == EProfilerSampleTypes::Memory)
 		{
-			const int32 LastChildIndex = ThreadNode.Children.Num() - 1;
-			ThreadNode.CyclesStart = ThreadNode.Children[0]->CyclesStart;
-			ThreadNode.CyclesEnd = ThreadNode.Children[LastChildIndex]->CyclesEnd;
-			ThreadNode.CycleCounterStartTimeMS = MetaData->ConvertCyclesToMS( ThreadNode.CyclesStart );
-			ThreadNode.CycleCounterEndTimeMS = MetaData->ConvertCyclesToMS( ThreadNode.CyclesEnd );
-
-			FProfilerSample& ProfilerSample = MutableCollection[ThreadNode.SampleIndex];
-			ProfilerSample.SetStartAndEndMS( MetaData->ConvertCyclesToMS( ThreadNode.CyclesStart ), MetaData->ConvertCyclesToMS( ThreadNode.CyclesEnd ) );
+			// @TODO: Remove later
+			TypedValue *= 1.0f / 1024.0f;
 		}
 	}
 
-	// Get the game thread time.
-	check( GameThreadFName != NAME_None );
-	const FProfilerStackNode& GameThreadNode = *ThreadNodes.FindChecked( GameThreadFName );
-	const double GameThreadStartMS = MetaData->ConvertCyclesToMS( GameThreadNode.CyclesStart );
-	const double GameThreadEndMS = MetaData->ConvertCyclesToMS( GameThreadNode.CyclesEnd );
-	MutableCollection[FrameRootSampleIndex].SetStartAndEndMS( GameThreadStartMS, GameThreadEndMS );
-	
-	// Advance frame
-	const uint32 LastFrameIndex = DataProvider->GetNumFrames();
-	DataProvider->AdvanceFrame( GameThreadEndMS - GameThreadStartMS );
- 
-	// Update aggregated stats
-	UpdateAggregatedStats( LastFrameIndex );
- 
-	// Update aggregated events.
-	UpdateAggregatedEventGraphData( LastFrameIndex );
-
-	// RootNode is the same as the game thread node.
-	out_ProfilerFrame.Root->CycleCounterStartTimeMS = GameThreadStartMS;
-	out_ProfilerFrame.Root->CycleCounterEndTimeMS = GameThreadEndMS;
-
-	for( auto It = ThreadNodes.CreateIterator(); It; ++It )
-	{
-		out_ProfilerFrame.AddChild( It.Value() );
-	}
-
-	out_ProfilerFrame.SortChildren();
+	_ValueOneFrame += TypedValue;
 }
\ No newline at end of file
diff --git a/Engine/Source/Developer/Profiler/Private/ProfilerSession.h b/Engine/Source/Developer/Profiler/Private/ProfilerSession.h
--- a/Engine/Source/Developer/Profiler/Private/ProfilerSession.h
+++ b/Engine/Source/Developer/Profiler/Private/ProfilerSession.h
@@ -448,7 +448,6 @@ private:
 				StatFNameDescriptions.Add( StatFName, StatPtr );
 			}
 
-			// @TODO yrx 2014-03-24 FindRef!
 			FProfilerGroup* GroupPtr = GroupDescriptions.FindRef( GroupID );
 
 			StatPtr->Initialize( StatName, GroupPtr, InType );
@@ -645,24 +644,7 @@ public:
 	};
 
 	/** Default constructor. */
-	FProfilerAggregatedStat( const FName& InStatName, const FName& InGroupName, EProfilerSampleTypes::Type InStatType )
-		: _StatName( InStatName )
-		, _GroupName( InGroupName )
-		, _ValueOneFrame( 0.0f )
-		, _ValueAllFrames( 0.0f )
-		, _MinValueAllFrames( FLT_MAX )
-		, _MaxValueAllFrames( FLT_MIN )
-	
-		, _NumCallsAllFrames( 0 )
-		, _NumCallsOneFrame( 0 )
-		, _MinNumCallsAllFrames( MAX_uint32 )
-		, _MaxNumCallsAllFrames( 0 )
-
-		, _NumFrames( 0 )
-		, _NumFramesWithCall( 0 )
-
-		, StatType( InStatType )
-	{}
+	FProfilerAggregatedStat( const FName& InStatName, const FName& InGroupName, EProfilerSampleTypes::Type InStatType );
 
 public:
 	/** @return The average value of all combined instances. */
@@ -721,124 +703,17 @@ public:
 	/**
 	 * @return a string representation of this aggregated stat.
 	 */
-	FString ToString() const
-	{
-		FString FormattedValueStr;
-		if( StatType == EProfilerSampleTypes::HierarchicalTime )
-		{
-			FormattedValueStr = FString::Printf( TEXT("{Value Min:%.3f Avg:%.3f Max:%.3f (MS) / Calls (%.1f%%) Min:%.1f Avg:%.1f Max:%.1f}"), 
-				MinValue(), AvgValue(), MaxValue(),
-				FramesWithCallPct(), MinNumCalls(), AvgNumCalls(), MaxNumCalls() );
-		}
-		else if( StatType == EProfilerSampleTypes::Memory )
-		{
-			FormattedValueStr = FString::Printf( TEXT("{Min:%.2f Avg:%.2f Max:%.2f (KB)}"), MinValue(), AvgValue(), MaxValue() );
-		}
-		else if( StatType == EProfilerSampleTypes::NumberInt )
-		{
-			FormattedValueStr = FString::Printf( TEXT("{Min:%.2f Avg:%.2f Max:%.2f}"), MinValue(), AvgValue(), MaxValue() );
-		}
-		else if( StatType == EProfilerSampleTypes::NumberFloat )
-		{
-			FormattedValueStr = FString::Printf( TEXT("{Min:%.2f Avg:%.2f Max:%.2f}"), MinValue(), AvgValue(), MaxValue() );
-		}
-		else
-		{
-			check(0);
-		}
-		return FormattedValueStr;
-	}
+	FString ToString() const;
 
 	
 	/**
 	 * @return a string representation of the specified value type.
 	 */
-	FString GetFormattedValue( const Type ValueType ) const
-	{
-		check( ValueType < Type::EInvalidOrMax );
-		const double ValueArray[Type::EInvalidOrMax] =
-		{
-			AvgValue(),
-			MinValue(),
-			MaxValue(),
-
-			AvgNumCalls(),
-			MinNumCalls(),
-			MaxNumCalls(),
-		
-			FramesWithCallPct(),
-		};
-
-		FString FormatValueStr;
-		if( StatType == EProfilerSampleTypes::HierarchicalTime )
-		{
-			if( ValueType == EMinValue || ValueType == EAvgValue || ValueType == EMaxValue )
-			{
-				FormatValueStr = FString::Printf( TEXT("%.3f (MS)"), ValueArray[ValueType] );
-			}
-			else if( ValueType == EFramesWithCallPct )
-			{
-				FormatValueStr = FString::Printf( TEXT("%.1f%%"), ValueArray[EFramesWithCallPct] );
-			}
-			else
-			{
-				FormatValueStr = FString::Printf( TEXT("%.1f"), ValueArray[ValueType] );
-			}
-		}
-		else if( StatType == EProfilerSampleTypes::Memory )
-		{
-			FormatValueStr = FString::Printf( TEXT("%.2f (KB)"), ValueArray[ValueType] );
-		}
-		else if( StatType == EProfilerSampleTypes::NumberInt )
-		{
-			FormatValueStr = FString::Printf( TEXT("%.2f"), ValueArray[ValueType] );
-		}
-		else if( StatType == EProfilerSampleTypes::NumberFloat )
-		{
-			FormatValueStr = FString::Printf( TEXT("%.2f"), ValueArray[ValueType] );
-		}
-		else
-		{
-			check(0);
-		}
-		return FormatValueStr;
-	}
+	FString GetFormattedValue( const Type ValueType ) const;
 
 protected:
 	/** Called once a frame to update aggregates. */
-	void Advance()
-	{
-		_NumFrames++;
-
-		_NumCallsAllFrames += _NumCallsOneFrame;
-		_ValueAllFrames += _ValueOneFrame;
-
-		// Calculate new extreme values.
-		if( StatType == EProfilerSampleTypes::HierarchicalTime )
-		{
-			// Update the extreme values only if this stat has been called at least once.
-			if( _NumCallsOneFrame > 0 )
-			{
-				_NumFramesWithCall ++;
-			}
-
-			{
-				_MinValueAllFrames = FMath::Min<float>( _MinValueAllFrames, _ValueOneFrame );
-				_MaxValueAllFrames = FMath::Max<float>( _MaxValueAllFrames, _ValueOneFrame );
-
-				_MinNumCallsAllFrames = FMath::Min<float>( _MinNumCallsAllFrames, _NumCallsOneFrame );
-				_MaxNumCallsAllFrames = FMath::Max<float>( _MaxNumCallsAllFrames, _NumCallsOneFrame );
-			}
-		}
-		else
-		{
-			_MinValueAllFrames = FMath::Min<float>( _MinValueAllFrames, _ValueOneFrame );
-			_MaxValueAllFrames = FMath::Max<float>( _MaxValueAllFrames, _ValueOneFrame );
-		}
-
-		_ValueOneFrame = 0.0f;
-		_NumCallsOneFrame = 0;
-	}
+	void Advance();
 
 	/**
 	 * Adds a profiler sample to our aggregated data
@@ -846,31 +721,7 @@ protected:
 	 * @param Sample - a reference to the profiler sample that will be aggregated
 	 *
 	 */
-	FProfilerAggregatedStat& operator+=( const FProfilerSample& Sample )
-	{
-		double TypedValue = 0.0;
-
-		// Determine whether to we are reading a time hierarchical sample or not.
-		if( Sample.Type() == EProfilerSampleTypes::HierarchicalTime )
-		{
-			TypedValue = Sample.DurationMS();
-			_NumCallsOneFrame += (uint32)Sample.CounterAsInt();
-		}
-		else
-		{
-			TypedValue = Sample.CounterAsFloat();
-
-			if( Sample.Type() == EProfilerSampleTypes::Memory )
-			{
-				// @TODO: Remove later
-				TypedValue *= 1.0f / 1024.0f;
-			}
-		}
-
-		_ValueOneFrame += TypedValue;
-		
-		return *this;
-	}
+	void Aggregate( const FProfilerSample& Sample, const FProfilerStatMetaDataRef& Metadata );
 
 protected:
 	FName _StatName;
@@ -912,6 +763,8 @@ protected:
 	const EProfilerSampleTypes::Type StatType;
 };
 
+//template <> struct TIsPODType<FProfilerAggregatedStat> { enum { Value = true }; };
+
 // @TODO: Standardize naming used in the profiler session/ session / session instance.
 
 /** Class that holds the profiler session information. */
@@ -1064,6 +917,11 @@ public:
 		return AggregatedStats.Find( StatID );
 	}
 
+	FORCEINLINE_DEBUGGABLE const TMap<uint32, FInclusiveTime>& GetInclusiveAggregateStackStats( const uint32 FrameIndex ) const
+	{
+		return InclusiveAggregateStackStats[FrameIndex];
+	}
+
 	const FEventGraphDataRef GetEventGraphDataTotal() const;
 	const FEventGraphDataRef GetEventGraphDataMaximum() const;
 	const FEventGraphDataRef GetEventGraphDataAverage() const;
@@ -1071,13 +929,7 @@ public:
 	/**
 	 * @return number of bytes allocated by this profiler session.
 	 */
-	const SIZE_T GetMemoryUsage() const
-	{
-		SIZE_T MemoryUsage = 0;
-		MemoryUsage += DataProvider->GetMemoryUsage();
-		MemoryUsage += StatMetaData->GetMemoryUsage();
-		return MemoryUsage;
-	}
+	const SIZE_T GetMemoryUsage() const;
 
 	/**
 	 * @return a new instance of the graph data source for the specified stat ID.
@@ -1085,12 +937,17 @@ public:
 	FGraphDataSourceRefConst CreateGraphDataSource( const uint32 InStatID );
 
 	/**
-	 * @return a new instance of the event graph data for the specified frame range and event graph type.
+	 * @return a new instance of the event graph container for the specified frame range.
 	 */
-	FEventGraphDataRef CreateEventGraphData( const uint32 FrameStartIndex, const uint32 FrameEndIndex, const EEventGraphTypes::Type EventGraphType );
+	FEventGraphContainer CreateEventGraphData( const uint32 FrameStartIndex, const uint32 FrameEndIndex );
 
-protected:
+	/** Combines event graphs for the specified frames range. */
+	FEventGraphData* CombineEventGraphs( const uint32 FrameStartIndex, const uint32 FrameEndIndex );
+
+	/** Combines event graphs for the specified frames range, as the task graph task. */
+	void CombineEventGraphsTask( const uint32 FrameStartIndex, const uint32 FrameEndIndex );
 
+protected:
 	/**
 	 * Recursively populates hierarchy based on the specified profiler cycle graph.
 	 *
@@ -1102,11 +959,12 @@ protected:
 	 * @return <describe return value>
 	 */
 	void PopulateHierarchy_Recurrent
-	( 
+	(
+		const uint32 StatThreadID,
 		const FProfilerCycleGraph& ParentGraph, 
-		const double ParentStartTimeMS, 
-		const double ParentDurationMS, 
-		const uint32 ParentSampleIndex 
+		const uint32 ParentDurationCycles, 
+		const uint32 ParentSampleIndex,
+		TMap<uint32, FInclusiveTime>& StatIDToInclusiveTime
 	);
 
 	/** Called when this profiler session receives a new profiler data. */
@@ -1142,10 +1000,7 @@ protected:
 	void LoadComplete();
 
 	/** Sets number of frames. */
-	void SetNumberOfFrames( int32 InNumFrames )
-	{
-		NumFrames = InNumFrames;
-	}
+	void SetNumberOfFrames( int32 InNumFrames );
 	
 	/**
 	 * @return progress as floating point between 0 and 1.
@@ -1156,7 +1011,8 @@ protected:
 	/** All aggregated stats, stored as StatID -> FProfilerAggregatedStat. */
 	TMap<uint32, FProfilerAggregatedStat > AggregatedStats;
 
-	TMap<uint32, FProfilerAggregatedStat > LastFrameStats;
+	/** Inclusive aggregated stack stack for all frames. FrameIndex < StatID, StatValue > */
+	TArray< TMap<uint32, FInclusiveTime> > InclusiveAggregateStackStats;
 
 	// TODO: This part will go away once we fully switch over to the stats2.
 	//{
@@ -1199,21 +1055,24 @@ protected:
 	/** Event graph completion sync ( combine for max, combine for add ) that can be done in parallel, but need to wait before combining the next frame. */
 	FGraphEventRef CompletionSync;
 
+	/** Combined event graphs calculated on the task graph threads. */
+	TLockFreePointerListFIFO<FEventGraphData, PLATFORM_CACHE_LINE_SIZE> CombinedSubEventGraphsLFL;
+
 	// TODO: Temporary, need to be stored in metadata or send via UMB.
 	/** The time when this profiler session was created ( time of the connection to the client, time when a capture file was created ). */
-	/*const*/ FDateTime CreationTime;
+	FDateTime CreationTime;
 
 	/** Session type for this profiler session. */
-	/*const*/ EProfilerSessionTypes::Type SessionType;
+	EProfilerSessionTypes::Type SessionType;
 
 	/** Shared pointer to the session instance info. */
 	const ISessionInstanceInfoPtr SessionInstanceInfo;
 
 	/** An unique session instance ID. */
-	/*const*/ FGuid SessionInstanceID;
+	FGuid SessionInstanceID;
 
 	/** Data filepath. */
-	/*const*/ FString DataFilepath;
+	FString DataFilepath;
 
 	/** Number of frames in the file. */	
 	int32 NumFrames;
@@ -1233,52 +1092,4 @@ protected:
 public:
 	/** Provides analysis of the frame rate */
 	TSharedRef<FFPSAnalyzer> FPSAnalyzer;
-};
-
-/*-----------------------------------------------------------------------------
-	Profiler session for the raw stats files
------------------------------------------------------------------------------*/
-
-struct FAllocationInfo;
-
-class FRawProfilerSession : public FProfilerSession
-{
-	friend class FProfilerManager;
-	friend class FProfilerActionManager;
-
-	/** Profiler stream that contains all read raw profiler frames. */
-	FProfilerStream ProfilerStream;
-
-	/** Stats thread state, mostly used to manage the stats metadata. */
-	FStatsLoadedState StatsThreadStats;
-	FStatsReadStream Stream;
-
-	/** Index of the last processed data for the mini-view. */
-	int32 CurrentMiniViewFrame;
-
-public:
-	/**
-	* Default constructor, creates a profiler session from a capture file.
-	*/
-	FRawProfilerSession( const FString& InRawStatsFileFileath );
-
-	/** Destructor. */
-	~FRawProfilerSession();
-
-	/** Updates this profiler session. */
-	bool HandleTicker( float DeltaTime );
-
-	/** Starts a process of loading the raw stats file. */
-	void PrepareLoading();
-
-	const FProfilerStream& GetStream() const
-	{
-		return ProfilerStream;
-	}
-
-	/**
-	 *	Process all stats packets and convert them to data accessible by the profiler.
-	 *	Temporary version, will be optimized later.
-	 */
-	void ProcessStatPacketArray( const FStatPacketArray& PacketArray, FProfilerFrame& out_ProfilerFrame, int32 FrameIndex );
 };
\ No newline at end of file
diff --git a/Engine/Source/Developer/Profiler/Private/ProfilerStream.cpp b/Engine/Source/Developer/Profiler/Private/ProfilerStream.cpp
--- a/Engine/Source/Developer/Profiler/Private/ProfilerStream.cpp
+++ b/Engine/Source/Developer/Profiler/Private/ProfilerStream.cpp
@@ -44,11 +44,11 @@ void FProfilerUIStream::GenerateUIStream( const FProfilerStream& ProfilerStream,
 	}
 	else
 	{
-		// @TODO yrx 2014-04-28 Combining profiler frames doesn't make sense at this moment
+		// #Profiler: 2014-04-28 Combining profiler frames doesn't make sense at this moment
 		// At this scale we may want to switch into line graph or something?
 	}
 
-	// @TODO yrx 2014-04-25 MAX_THREAD_DEPTH;
+	// #Profiler: 2014-04-25 MAX_THREAD_DEPTH;
 	LinearizeStream();
 }
 
diff --git a/Engine/Source/Developer/Profiler/Private/ProfilerStream.h b/Engine/Source/Developer/Profiler/Private/ProfilerStream.h
--- a/Engine/Source/Developer/Profiler/Private/ProfilerStream.h
+++ b/Engine/Source/Developer/Profiler/Private/ProfilerStream.h
@@ -6,6 +6,23 @@
 	Basic structures
 -----------------------------------------------------------------------------*/
 
+/** Helper struct used to calculate inclusive times, ignoring recursive calls. */
+struct FInclusiveTime
+{
+	/** Duration in cycles. */
+	uint32 DurationCycles;
+	/** Number of calls. */
+	int32 CallCount;
+	/** Number of recursion. */
+	int32 Recursion;
+	FInclusiveTime()
+		: DurationCycles( 0.0 )
+		, CallCount( 0 )
+		, Recursion( 0 )
+	{}
+};
+
+
 /** Profiler stack node, used to store the whole call stack for one frame. */
 struct FProfilerStackNode
 {
@@ -19,7 +36,7 @@ struct FProfilerStackNode
 	int64 CyclesStart;
 	int64 CyclesEnd;
 
-	// @TODO yrx 2014-04-23 This should be based on the global time.
+	// #Profiler: 2014-04-23 This should be based on the global time.
 	double CycleCounterStartTimeMS;
 	double CycleCounterEndTimeMS;
 
@@ -103,7 +120,7 @@ protected:
 
 public:
 	/** Root node. */
-	// @TODO yrx 2014-04-25 Replace with TIndirectArray<FProfilerStackNode>??
+	// #Profiler: 2014-04-25 Replace with TIndirectArray<FProfilerStackNode>??
 	FProfilerStackNode* Root;
 
 	/** Thread times in milliseconds for this frame. */
@@ -125,7 +142,7 @@ public:
 	 */
 	double LastAccessTime;
 
-	// @TODO yrx 2014-04-22 Non-frame stats like accumulator, counters, memory etc?
+	// #Profiler: 2014-04-22 Non-frame stats like accumulator, counters, memory etc?
 
 	/**
 	*	Indicates whether this profiler frame is in the memory.
@@ -206,14 +223,14 @@ public:
 };
 
 /** Contains all processed profiler's frames. */
-// @TODO yrx 2014-04-22 Array frames to real frames conversion.
-// @TODO yrx 2014-04-22 Can be done fully lock free and thread safe, but later.
+// #Profiler: 2014-04-22 Array frames to real frames conversion.
+// #Profiler: 2014-04-22 Can be done fully lock free and thread safe, but later.
 struct FProfilerStream
 {
 protected:
 
 	/** History frames collected so far or read from the file. */
-	// @TODO yrx 2014-04-24 TSharedRef<THREAD_SAFE> ?
+	// #Profiler: 2014-04-24 TSharedRef<THREAD_SAFE> ?
 	TChunkedArray<FProfilerFrame*> Frames;
 
 	/** Each element in this array stores the frame time, accessed by a frame index, in millisecond. */
@@ -332,7 +349,7 @@ public:
 	}
 };
 
-// @TODO yrx 2014-04-23 ProfilerFrameLOD?
+// #Profiler: 2014-04-23 ProfilerFrameLOD?
 
 /**
 *	Profiler UI stack node.
@@ -493,7 +510,7 @@ struct FProfilerUIStream
 		 *	Default number of rows displaying cycle counters.
 		 *	Read it as the call stack depth.
 		 */
-		// @TODO yrx 2014-04-25 This probably should be the max seen so far for each thread.
+		// #Profiler: 2014-04-25 This probably should be the max seen so far for each thread.
 		DEFAULT_VISIBLE_THREAD_DEPTH = 16,
 	};
 
diff --git a/Engine/Source/Developer/Profiler/Private/Widgets/SDataGraph.cpp b/Engine/Source/Developer/Profiler/Private/Widgets/SDataGraph.cpp
--- a/Engine/Source/Developer/Profiler/Private/Widgets/SDataGraph.cpp
+++ b/Engine/Source/Developer/Profiler/Private/Widgets/SDataGraph.cpp
@@ -1156,32 +1156,6 @@ FReply SDataGraph::OnMouseWheel( const FGeometry& MyGeometry, const FPointerEven
 	return FReply::Handled();
 }
 
-const FEventGraphDataHandlerRef SDataGraph::PrepareEventGraphDataHandler( const FVector2D& ScreenSpacePosition )
-{
-	const FTrackedStatPtr FirstTrackedStat = GetFirstGraph();
-	if (FirstTrackedStat.IsValid())
-	{
-		if( ViewMode == EDataGraphViewModes::Time )
-		{
-			TMap<FGuid,uint32> StartIndices;
-			const float TimeAccuracyMS = FTimeAccuracy::AsFrameTime( TimeBasedAccuracy );
-			// #YRX_ProfilerFix? 2015-12-08 
-			//GraphDescription->GraphDataSource->GetStartIndicesFromTimeRange( FrameTimesMS[0], FrameTimesMS[1], StartIndices );
-
-			const FEventGraphDataHandlerRef Params = MakeShareable( new FEventGraphDataHandler( /*StartIndices*/ ) ); 
-			return Params;
-		}
-		else if( ViewMode == EDataGraphViewModes::Index )
-		{
-			const FGraphDataSourceRefConst& GraphDataSource = FirstTrackedStat->GraphDataSource;
-			const FEventGraphDataHandlerRef Params = MakeShareable( new FEventGraphDataHandler( GraphDataSource->GetSessionInstanceID(), FrameIndices[0], FrameIndices[1], EEventGraphTypes::Maximum ) ); 
-			return Params;
-		}
-	}
-
-	return MakeShareable( new FEventGraphDataHandler() );
-}
-
 FReply SDataGraph::OnMouseButtonDoubleClick( const FGeometry& MyGeometry, const FPointerEvent& MouseEvent )
 {
 	return FReply::Unhandled();
@@ -1259,8 +1233,6 @@ void SDataGraph::ShowContextMenu(const FVector2D& ScreenSpacePosition, const FPo
 	const FProfilerActionManager& ProfilerActionManager = FProfilerManager::GetActionManager();
 
 	// Build data required for opening event graph(s).
-	const FEventGraphDataHandlerRef EventGraphDataHandler = PrepareEventGraphDataHandler( ScreenSpacePosition );
-
 	const bool bShouldCloseWindowAfterMenuSelection = true;
 	FMenuBuilder MenuBuilder( bShouldCloseWindowAfterMenuSelection, ProfilerCommandList );
 	
diff --git a/Engine/Source/Developer/Profiler/Private/Widgets/SDataGraph.h b/Engine/Source/Developer/Profiler/Private/Widgets/SDataGraph.h
--- a/Engine/Source/Developer/Profiler/Private/Widgets/SDataGraph.h
+++ b/Engine/Source/Developer/Profiler/Private/Widgets/SDataGraph.h
@@ -42,9 +42,6 @@ namespace EDataGraphMultiModes
 /** The delegate to be invoked when the frame offset has been changed. */
 DECLARE_DELEGATE_OneParam( FGraphOffsetChangedDelegate, int /*InFrameOffset*/ );
 
-/** The delegate to be invoked when the user clicked on the data graph. */
-DECLARE_DELEGATE_OneParam( FDataGraphClickedDelegate, FEventGraphDataHandlerRef /*EventGraphDataHandler*/ );
-
 /** The delegate to be invoked when the data graph summary widget wants to know index of the frame currently being hovered by the mouse. */
 DECLARE_DELEGATE_RetVal( int32, FGetHoveredFrameIndexDelegate );
 
@@ -301,8 +298,6 @@ protected:
 
 	void UpdateState();
 
-	const FEventGraphDataHandlerRef PrepareEventGraphDataHandler( const FVector2D& ScreenSpacePosition );
-
 	/** Handles FExecuteAction for ViewMode_SetIndexBased. */
 	void ViewMode_SetIndexBased_Execute();
 	/** Handles FCanExecuteAction for ViewMode_SetIndexBased. */
diff --git a/Engine/Source/Developer/Profiler/Private/Widgets/SEventGraph.cpp b/Engine/Source/Developer/Profiler/Private/Widgets/SEventGraph.cpp
--- a/Engine/Source/Developer/Profiler/Private/Widgets/SEventGraph.cpp
+++ b/Engine/Source/Developer/Profiler/Private/Widgets/SEventGraph.cpp
@@ -53,106 +53,125 @@ namespace EEventGraphViewModes
 	}
 }
 
-namespace EventGraphColumns
+struct FEventGraphColumns
 {
-	// @TODO: Needs refactoring
-
-	/** Contains basic information about columns used in the event graph widget. Names should be localized. */
-	static const FEventGraphColumn Collection[ 7 ] =
+	/** Default constructor. */
+	FEventGraphColumns()
+		: NumColumns( EEventPropertyIndex::None+1 )
 	{
-		// Generated from a XLSX file.
-		FEventGraphColumn
+		// Make event property is initialized.
+		FEventGraphSample::InitializePropertyManagement();
+
+		Collection = new FEventGraphColumn[NumColumns];
+
+		Collection[EEventPropertyIndex::StatName] = FEventGraphColumn
 		(
 			EEventPropertyIndex::StatName,
-			TEXT("name"),
-			LOCTEXT("EventNameColumnTitle", "Event Name"),
-			LOCTEXT("EventNameColumnDesc", "Name of the event"),
-			false,true,true,false,false,
+			TEXT( "name" ),
+			LOCTEXT( "EventNameColumnTitle", "Event Name" ),
+			LOCTEXT( "EventNameColumnDesc", "Name of the event" ),
+			false, true, true, false, false,
 			HAlign_Left,
 			0.0f
-		),
-		
-		FEventGraphColumn
+		);
+
+		Collection[EEventPropertyIndex::InclusiveTimeMS] = FEventGraphColumn
 		(
 			EEventPropertyIndex::InclusiveTimeMS,
-			TEXT("inc"),
-			LOCTEXT("InclusiveTimeMSTitle", "Inc Time (MS)"),
-			LOCTEXT("InclusiveTimeMSDesc", "Duration of the sample and its children, in milliseconds"),
-			false,true,true,true,true,
+			TEXT( "inc" ),
+			LOCTEXT( "InclusiveTimeMSTitle", "Inc Time (MS)" ),
+			LOCTEXT( "InclusiveTimeMSDesc", "Duration of the sample and its children, in milliseconds" ),
+			false, true, true, true, true,
 			HAlign_Right,
-			0.0f
-		),
-		
-		FEventGraphColumn
+			72.0f
+		);
+
+		Collection[EEventPropertyIndex::InclusiveTimePct] = FEventGraphColumn
 		(
 			EEventPropertyIndex::InclusiveTimePct,
-			TEXT("inc%"),
-			LOCTEXT("InclusiveTimePercentageTitle", "Inc Time (%)"),
-			LOCTEXT("InclusiveTimePercentageDesc", "Duration of the sample and its children as percent of the caller"),
-			true,true,true,false,false,
+			TEXT( "inc%" ),
+			LOCTEXT( "InclusiveTimePercentageTitle", "Inc Time (%)" ),
+			LOCTEXT( "InclusiveTimePercentageDesc", "Duration of the sample and its children as percent of the caller" ),
+			true, true, true, false, false,
 			HAlign_Right,
-			64.0f
-		),
-		
-		FEventGraphColumn
+			72.0f
+		);
+
+		Collection[EEventPropertyIndex::ExclusiveTimeMS] = FEventGraphColumn
 		(
 			EEventPropertyIndex::ExclusiveTimeMS,
-			TEXT("exc"),
-			LOCTEXT("ExclusiveTimeMSTitle", "Exc Time (MS)"),
-			LOCTEXT("ExclusiveTimeMSDesc", "Exclusive time of this event, in milliseconds"),
-			false,true,true,true,false,
+			TEXT( "exc" ),
+			LOCTEXT( "ExclusiveTimeMSTitle", "Exc Time (MS)" ),
+			LOCTEXT( "ExclusiveTimeMSDesc", "Exclusive time of this event, in milliseconds" ),
+			false, true, true, true, false,
 			HAlign_Right,
-			0.0f
-		),
-		
-		FEventGraphColumn
+			72.0f
+		);
+
+		Collection[EEventPropertyIndex::ExclusiveTimePct] = FEventGraphColumn
 		(
 			EEventPropertyIndex::ExclusiveTimePct,
-			TEXT("exc%"),
-			LOCTEXT("ExclusiveTimePercentageTitle", "Exc Time (%)"),
-			LOCTEXT("ExclusiveTimePercentageDesc", "Exclusive time of this event as percent of this call's inclusive time"),
-			true,true,true,false,false,
+			TEXT( "exc%" ),
+			LOCTEXT( "ExclusiveTimePercentageTitle", "Exc Time (%)" ),
+			LOCTEXT( "ExclusiveTimePercentageDesc", "Exclusive time of this event as percent of this call's inclusive time" ),
+			true, true, true, false, false,
 			HAlign_Right,
-			64.0f
-		),
-		
-		FEventGraphColumn
+			72.0f
+		);
+
+		Collection[EEventPropertyIndex::NumCallsPerFrame] = FEventGraphColumn
 		(
 			EEventPropertyIndex::NumCallsPerFrame,
-			TEXT("calls"),
-			LOCTEXT("CallsPerFrameTitle", "Calls"),
-			LOCTEXT("CallsPerFrameDesc", "Number of times this event was called"),
-			false,true,true,true,false,
+			TEXT( "calls" ),
+			LOCTEXT( "CallsPerFrameTitle", "Calls" ),
+			LOCTEXT( "CallsPerFrameDesc", "Number of times this event was called" ),
+			false, true, true, true, false,
 			HAlign_Right,
-			64.0f
-		),
+			48.0f
+		);
 
 		// Fake column used as a default column for NAME_None
-		FEventGraphColumn
+		Collection[EEventPropertyIndex::None] = FEventGraphColumn
 		(
 			EEventPropertyIndex::None,
-			TEXT("None"),
-			LOCTEXT("None", "None"),
-			LOCTEXT("None", "None"),
-			false,false,false,false,false,
+			TEXT( "None" ),
+			LOCTEXT( "None", "None" ),
+			LOCTEXT( "None", "None" ),
+			false, false, false, false, false,
 			HAlign_Left,
 			0.0f
-		),
-	};
+		);
 
-	static const uint32 NumColumns = ARRAY_COUNT(Collection);
+		ColumnNameToIndexMapping = TMapBuilder<FName, const FEventGraphColumn*>()
+			.Add( TEXT( "StatName" ), &Collection[EEventPropertyIndex::StatName] )
+			.Add( TEXT( "InclusiveTimeMS" ), &Collection[EEventPropertyIndex::InclusiveTimeMS] )
+			.Add( TEXT( "InclusiveTimePct" ), &Collection[EEventPropertyIndex::InclusiveTimePct] )
+			.Add( TEXT( "ExclusiveTimeMS" ), &Collection[EEventPropertyIndex::ExclusiveTimeMS] )
+			.Add( TEXT( "ExclusiveTimePct" ), &Collection[EEventPropertyIndex::ExclusiveTimePct] )
+			.Add( TEXT( "NumCallsPerFrame" ), &Collection[EEventPropertyIndex::NumCallsPerFrame] )
+			.Add( NAME_None, &Collection[EEventPropertyIndex::None] )
+			;
+	}
+
+	~FEventGraphColumns()
+	{
+		delete[] Collection;
+	}
+
+	/** Contains basic information about columns used in the event graph widget. Names should be localized. */
+	FEventGraphColumn* Collection;
+
+	const uint32 NumColumns;
 
 	// Generated from a XLSX file.
-	static const TMap<FName,const FEventGraphColumn*> ColumnNameToIndexMapping = TMapBuilder<FName, const FEventGraphColumn*>()
-		.Add(TEXT("StatName"),&Collection[EEventPropertyIndex::StatName])
-		.Add(TEXT("InclusiveTimeMS"),&Collection[EEventPropertyIndex::InclusiveTimeMS])
-		.Add(TEXT("InclusiveTimePct"),&Collection[EEventPropertyIndex::InclusiveTimePct])
-		.Add(TEXT("ExclusiveTimeMS"),&Collection[EEventPropertyIndex::ExclusiveTimeMS])
-		.Add(TEXT("ExclusiveTimePct"),&Collection[EEventPropertyIndex::ExclusiveTimePct])
-		.Add(TEXT("NumCallsPerFrame"),&Collection[EEventPropertyIndex::NumCallsPerFrame])
-		.Add(NAME_None,&Collection[EEventPropertyIndex::None])
-		;
-}
+	TMap<FName, const FEventGraphColumn*> ColumnNameToIndexMapping;
+
+	static const FEventGraphColumns& Get()
+	{
+		static FEventGraphColumns Instance;
+		return Instance;
+	}
+};
 
 /*-----------------------------------------------------------------------------
 	SEventTreeItem
@@ -194,7 +213,7 @@ protected:
 	BEGIN_SLATE_FUNCTION_BUILD_OPTIMIZATION
 	TSharedRef<SWidget> GenerateWidgetForColumnID( const FName& InColumnID, const bool bIsEventNameColumn, const TSharedRef<class ITableRow>& TableRow )
 	{
-		const FEventGraphColumn& Column = *EventGraphColumns::ColumnNameToIndexMapping.FindChecked( InColumnID );
+		const FEventGraphColumn& Column = *FEventGraphColumns::Get().ColumnNameToIndexMapping.FindChecked( InColumnID );
 		
 		if( bIsEventNameColumn )
 		{
@@ -460,7 +479,7 @@ public:
 			SNew( SEventGraphTableCell, SharedThis(this), OwnerEventGraph )
 			.Visibility( this, &SEventGraphTableRow::IsColumnVisible, ColumnID )
 			.ColumnID( ColumnID )
-			.IsEventNameColumn( ColumnID == EventGraphColumns::Collection[EEventPropertyIndex::StatName].ID )
+			.IsEventNameColumn( ColumnID == FEventGraphColumns::Get().Collection[EEventPropertyIndex::StatName].ID )
 			.EventPtr( EventPtr )
 			.OnSetHoveredTableCell( this, &SEventGraphTableRow::OnSetHoveredTableCell )
 		];
@@ -1322,7 +1341,7 @@ void SEventGraph::SortEvents()
 
 	if( ColumnBeingSorted != NAME_None )
 	{
-		const FEventGraphColumn& Column = *EventGraphColumns::ColumnNameToIndexMapping.FindChecked( ColumnBeingSorted );
+		const FEventGraphColumn& Column = *FEventGraphColumns::Get().ColumnNameToIndexMapping.FindChecked( ColumnBeingSorted );
 
 		if(	GetCurrentStateViewMode() == EEventGraphViewModes::Hierarchical || 
 			GetCurrentStateViewMode() == EEventGraphViewModes::FlatInclusive || 
@@ -1330,9 +1349,9 @@ void SEventGraph::SortEvents()
 			GetCurrentStateViewMode() == EEventGraphViewModes::FlatExclusive || 
 			GetCurrentStateViewMode() == EEventGraphViewModes::FlatExclusiveCoalesced )
 		{
-			FEventArraySorter::Sort( GetCurrentState()->GetRealRoot()->GetChildren(), Column.ID, EColumnSortModeToEventCompareOp(ColumnSortMode) );
-			FEventArraySorter::Sort( Events_FlatCoalesced, Column.ID, EColumnSortModeToEventCompareOp(ColumnSortMode) );
-			FEventArraySorter::Sort( Events_Flat, Column.ID, EColumnSortModeToEventCompareOp(ColumnSortMode) );
+				FEventArraySorter::Sort( GetCurrentState()->GetRealRoot()->GetChildren(), Column.ID, EColumnSortModeToEventCompareOp( ColumnSortMode ) );
+				FEventArraySorter::Sort( Events_FlatCoalesced, Column.ID, EColumnSortModeToEventCompareOp( ColumnSortMode ) ); 
+				FEventArraySorter::Sort( Events_Flat, Column.ID, EColumnSortModeToEventCompareOp( ColumnSortMode ) ); 
 			
 			// Update not culled children.
 			GetCurrentState()->GetRoot()->SetBooleanStateForAllChildren<EEventPropertyIndex::bNeedNotCulledChildrenUpdate>(true);
@@ -1385,7 +1404,7 @@ bool SEventGraph::FilteringSearchBox_IsEnabled() const
 
 TSharedPtr<SWidget> SEventGraph::EventGraph_GetMenuContent() const
 {
-	const FEventGraphColumn& Column = *EventGraphColumns::ColumnNameToIndexMapping.FindChecked( HoveredColumnID );
+	const FEventGraphColumn& Column = *FEventGraphColumns::Get().ColumnNameToIndexMapping.FindChecked( HoveredColumnID );
 	const TArray<FEventGraphSamplePtr> SelectedEvents = TreeView_Base->GetSelectedItems();
 	const int NumSelectedEvents = SelectedEvents.Num();
 	FEventGraphSamplePtr SelectedEvent = NumSelectedEvents ? SelectedEvents[0] : nullptr;
@@ -1893,7 +1912,7 @@ void SEventGraph::TreeView_Refresh()
 
 void SEventGraph::TreeViewHeaderRow_CreateColumnArgs( const uint32 ColumnIndex )
 {
-	const FEventGraphColumn Column = EventGraphColumns::Collection[ColumnIndex];
+	const FEventGraphColumn Column = FEventGraphColumns::Get().Collection[ColumnIndex];
 	SHeaderRow::FColumn::FArguments ColumnArgs;
 
 	ColumnArgs
@@ -1932,9 +1951,9 @@ void SEventGraph::TreeViewHeaderRow_CreateColumnArgs( const uint32 ColumnIndex )
 void SEventGraph::InitializeAndShowHeaderColumns()
 {
 	ColumnSortMode = EColumnSortMode::Descending;
-	ColumnBeingSorted = EventGraphColumns::Collection[EEventPropertyIndex::InclusiveTimeMS].ID;
+	ColumnBeingSorted = FEventGraphColumns::Get().Collection[EEventPropertyIndex::InclusiveTimeMS].ID;
 
-	for( uint32 ColumnIndex = 0; ColumnIndex < EventGraphColumns::NumColumns; ColumnIndex++ )
+	for (uint32 ColumnIndex = 0; ColumnIndex < FEventGraphColumns::Get().NumColumns; ColumnIndex++)
 	{
 		TreeViewHeaderRow_CreateColumnArgs( ColumnIndex );
 	}
@@ -2113,7 +2132,7 @@ void SEventGraph::ContextMenu_ExpandHotPath_Execute()
 	FEventGraphSamplePtr EventPtr = SelectedItems[0];
 
 	ColumnSortMode = EColumnSortMode::Descending;
-	ColumnBeingSorted = EventGraphColumns::Collection[EEventPropertyIndex::InclusiveTimeMS].ID;
+	ColumnBeingSorted = FEventGraphColumns::Get().Collection[EEventPropertyIndex::InclusiveTimeMS].ID;
 	SortEvents();
 	
 	// Clear hot path
@@ -2153,9 +2172,9 @@ void SEventGraph::ContextMenu_CopySelectedToClipboard_Execute()
 	FString Result;
 
 	// Prepare header.
-	for( uint32 ColumnIndex = 0; ColumnIndex < EventGraphColumns::NumColumns; ColumnIndex++ )
+	for (uint32 ColumnIndex = 0; ColumnIndex < FEventGraphColumns::Get().NumColumns; ColumnIndex++)
 	{
-		const FEventGraphColumn& Column = EventGraphColumns::Collection[ColumnIndex];
+		const FEventGraphColumn& Column = FEventGraphColumns::Get().Collection[ColumnIndex];
 		Result += FString::Printf( TEXT("\"%s\","), *Column.ShortName.ToString() );
 	}
 	Result += LINE_TERMINATOR;
@@ -2165,9 +2184,9 @@ void SEventGraph::ContextMenu_CopySelectedToClipboard_Execute()
 	{
 		const FEventGraphSamplePtr EventPtr = SelectedEvents[EventIndex];
 		
-		for( uint32 ColumnIndex = 0; ColumnIndex < EventGraphColumns::NumColumns; ColumnIndex++ )
+		for (uint32 ColumnIndex = 0; ColumnIndex < FEventGraphColumns::Get().NumColumns; ColumnIndex++)
 		{
-			const FEventGraphColumn& Column = EventGraphColumns::Collection[ColumnIndex];
+			const FEventGraphColumn& Column = FEventGraphColumns::Get().Collection[ColumnIndex];
 			if( Column.Index != EEventPropertyIndex::None )
 			{
 				const FString FormattedValue = EventPtr->GetFormattedValue( Column.Index );
@@ -2250,11 +2269,11 @@ bool SEventGraph::ContextMenu_SortMode_IsChecked( const EColumnSortMode::Type In
 void SEventGraph::ContextMenu_ResetColumns_Execute()
 {
 	ColumnSortMode = EColumnSortMode::Descending;
-	ColumnBeingSorted = EventGraphColumns::Collection[EEventPropertyIndex::InclusiveTimeMS].ID;
+	ColumnBeingSorted = FEventGraphColumns::Get().Collection[EEventPropertyIndex::InclusiveTimeMS].ID;
 
-	for( uint32 ColumnIndex = 0; ColumnIndex < EventGraphColumns::NumColumns; ColumnIndex++ )
+	for( uint32 ColumnIndex = 0; ColumnIndex < FEventGraphColumns::Get().NumColumns; ColumnIndex++ )
 	{
-		const FEventGraphColumn& DefaultColumn = EventGraphColumns::Collection[ColumnIndex];
+		const FEventGraphColumn& DefaultColumn = FEventGraphColumns::Get().Collection[ColumnIndex];
 		const FEventGraphColumn& CurrentColumn = TreeViewHeaderColumns.FindChecked( DefaultColumn.ID );
 
 		if( DefaultColumn.bIsVisible && !CurrentColumn.bIsVisible )
@@ -2326,6 +2345,8 @@ bool SEventGraph::HeaderMenu_SortMode_IsChecked( const FName ColumnID, const ECo
 
 void SEventGraph::CreateEvents()
 {
+	// #YRX_Profiler 2015-12-17  Fix recursive calls accounting.
+
 	// Linear
 	GetCurrentState()->GetRoot()->GetLinearEvents( Events_Flat, true );
 
@@ -2372,17 +2393,17 @@ void SEventGraph::ShowEventsInViewMode( const TArray<FEventGraphSamplePtr>& Even
 	SetHierarchicalSelectedEvents( EventGraphState->SelectedEvents );
 	SetHierarchicalExpandedEvents( EventGraphState->ExpandedEvents );
 		
-	EEventPropertyIndex::Type ColumnIndex = EventGraphColumns::ColumnNameToIndexMapping.FindChecked( ColumnBeingSorted )->Index;
+	EEventPropertyIndex::Type ColumnIndex = FEventGraphColumns::Get().ColumnNameToIndexMapping.FindChecked( ColumnBeingSorted )->Index;
 
 	if( NewViewMode == EEventGraphViewModes::FlatInclusive || NewViewMode == EEventGraphViewModes::FlatInclusiveCoalesced || NewViewMode == EEventGraphViewModes::Hierarchical )
 	{
 		ColumnIndex = EEventPropertyIndex::InclusiveTimeMS;
-		SetSortModeForColumn( EventGraphColumns::Collection[ColumnIndex].ID, EColumnSortMode::Descending );
+		SetSortModeForColumn( FEventGraphColumns::Get().Collection[ColumnIndex].ID, EColumnSortMode::Descending );
 	}
 	else if( NewViewMode == EEventGraphViewModes::FlatExclusive || NewViewMode == EEventGraphViewModes::FlatExclusiveCoalesced )
 	{
 		ColumnIndex = EEventPropertyIndex::ExclusiveTimeMS;
-		SetSortModeForColumn( EventGraphColumns::Collection[ColumnIndex].ID, EColumnSortMode::Descending );
+		SetSortModeForColumn( FEventGraphColumns::Get().Collection[ColumnIndex].ID, EColumnSortMode::Descending );
 	}
 
 	ScrollToTheSlowestSelectedEvent( ColumnIndex );
@@ -2413,7 +2434,7 @@ void SEventGraph::ScrollToTheSlowestSelectedEvent( EEventPropertyIndex::Type Col
 	if( SelectedEvents.Num() > 0 )
 	{
 		// Sort events by the inclusive or the exclusive time, depends on the view mode.
-		const FEventGraphColumn& Column = EventGraphColumns::Collection[ColumnIndex];
+		const FEventGraphColumn& Column = FEventGraphColumns::Get().Collection[ColumnIndex];
 		FEventArraySorter::Sort( SelectedEvents, Column.ID, EEventCompareOps::Greater );
 
 		// Scroll to the the slowest item.
@@ -2628,7 +2649,7 @@ void SEventGraph::GenerateTopEvents( const TSet< FEventGraphSamplePtr >& EventPt
 	}
 
 	// Sort events by the inclusive time.
-	const FEventGraphColumn& Column = EventGraphColumns::Collection[EEventPropertyIndex::InclusiveTimeMS];
+	const FEventGraphColumn& Column = FEventGraphColumns::Get().Collection[EEventPropertyIndex::InclusiveTimeMS];
 	FEventArraySorter::Sort( EventPtrArray, Column.ID, EEventCompareOps::Greater );
 
 	// Calculate total time for the top events.
@@ -3329,21 +3350,27 @@ FText SEventGraph::FEventGraphState::GetHistoryDesc() const
 	return Result;
 }
 
-void CreateOneToOneMapping_EventGraphSample( const FEventGraphSamplePtr LocalEvent, const FEventGraphSamplePtr SourceEvent, TMap< FEventGraphSamplePtr, FEventGraphSamplePtr >& out_Mapping )
+static void CreateOneToOneMapping_EventGraphSample
+( 
+	const FEventGraphSamplePtr LocalEvent, 
+	const FEventGraphSamplePtr SourceEvent, 
+	TMap< FEventGraphSamplePtr, FEventGraphSamplePtr >& out_MaximumToAverageMapping,
+	TMap< FEventGraphSamplePtr, FEventGraphSamplePtr >& out_AverageToMaximumMapping
+)
 {
-	out_Mapping.Add( LocalEvent, SourceEvent );
+	out_MaximumToAverageMapping.Add( LocalEvent, SourceEvent );
+	out_AverageToMaximumMapping.Add( SourceEvent, LocalEvent );
 
-	check( LocalEvent->GetChildren().Num() == SourceEvent->GetChildren().Num() )
+	check( LocalEvent->GetChildren().Num() == SourceEvent->GetChildren().Num() );
 	for( int32 Index = 0; Index < LocalEvent->GetChildren().Num(); ++Index )
 	{
-		CreateOneToOneMapping_EventGraphSample( LocalEvent->GetChildren()[Index], SourceEvent->GetChildren()[Index], out_Mapping );
+		CreateOneToOneMapping_EventGraphSample( LocalEvent->GetChildren()[Index], SourceEvent->GetChildren()[Index], out_MaximumToAverageMapping, out_AverageToMaximumMapping );
 	}
 }
 
 void SEventGraph::FEventGraphState::CreateOneToOneMapping()
 {
-	CreateOneToOneMapping_EventGraphSample( MaximumEventGraph->GetRoot(), AverageEventGraph->GetRoot(), MaximumToAverageMapping );
-	CreateOneToOneMapping_EventGraphSample( AverageEventGraph->GetRoot(), MaximumEventGraph->GetRoot(), AverageToMaximumMapping );
+	CreateOneToOneMapping_EventGraphSample( MaximumEventGraph->GetRoot(), AverageEventGraph->GetRoot(), MaximumToAverageMapping, AverageToMaximumMapping );
 }
 
 #undef LOCTEXT_NAMESPACE
diff --git a/Engine/Source/Developer/Profiler/Private/Widgets/SEventGraph.h b/Engine/Source/Developer/Profiler/Private/Widgets/SEventGraph.h
--- a/Engine/Source/Developer/Profiler/Private/Widgets/SEventGraph.h
+++ b/Engine/Source/Developer/Profiler/Private/Widgets/SEventGraph.h
@@ -60,13 +60,17 @@ namespace EEventGraphViewModes
 	FName ToBrushName( const Type EventGraphViewMode );
 }
 
-// @TODO: Sort properties by size.
-
 /** Holds information about a column in the event graph widget. */
-class FEventGraphColumn /*: public FNoncopyable*/
+class FEventGraphColumn
 {
-public:
+	friend struct FEventGraphColumns;
+
+protected:
+	/** Default constructor, only used in FEventGraphColumns  */
+	FEventGraphColumn()
+	{}
 
+	/** Initialization constructor, only used in FEventGraphColumns. */
 	FEventGraphColumn
 	(
 		const EEventPropertyIndex::Type InIndex,
diff --git a/Engine/Source/Developer/Profiler/Private/Widgets/SProfilerGraphPanel.cpp b/Engine/Source/Developer/Profiler/Private/Widgets/SProfilerGraphPanel.cpp
--- a/Engine/Source/Developer/Profiler/Private/Widgets/SProfilerGraphPanel.cpp
+++ b/Engine/Source/Developer/Profiler/Private/Widgets/SProfilerGraphPanel.cpp
@@ -2,7 +2,6 @@
 
 #include "ProfilerPrivatePCH.h"
 
-// @TODO yrx 2014-04-18 Move SProfilerWindow later
 #define LOCTEXT_NAMESPACE "SProfilerGraphPanel"
 
 SProfilerGraphPanel::SProfilerGraphPanel()
@@ -133,14 +132,12 @@ void SProfilerGraphPanel::HorizontalScrollBar_OnUserScrolled( float ScrollOffset
 
 void SProfilerGraphPanel::VerticalScrollBar_OnUserScrolled( float ScrollOffset )
 {
-	// @TODO yrx 2014-04-24 Broadcast?
 	if( ViewMode == EProfilerViewMode::LineIndexBased )
 	{
 
 	}
 	else if( ViewMode == EProfilerViewMode::ThreadViewTimeBased )
 	{
-		// @TODO yrx 2014-04-23 
 		ThreadView->SetPositonYTo( ScrollOffset );
 	}
 }
@@ -257,7 +254,6 @@ void SProfilerGraphPanel::UpdateInternals()
 	}
 	else if( ViewMode == EProfilerViewMode::ThreadViewTimeBased )
 	{
-		// @TODO yrx 2014-04-23 
 	}
 }
 
diff --git a/Engine/Source/Developer/Profiler/Private/Widgets/SProfilerMiniView.cpp b/Engine/Source/Developer/Profiler/Private/Widgets/SProfilerMiniView.cpp
--- a/Engine/Source/Developer/Profiler/Private/Widgets/SProfilerMiniView.cpp
+++ b/Engine/Source/Developer/Profiler/Private/Widgets/SProfilerMiniView.cpp
@@ -123,7 +123,7 @@ int32 SProfilerMiniView::OnPaint( const FPaintArgs& Args, const FGeometry& Allot
 	if( IsReady() )
 	{
 		static const FSlateColorBrush SolidWhiteBrush = FSlateColorBrush( FColorList::White );
-		// @TODO yrx 2014-04-24 move to the global scope.
+		// #Profiler 2014-04-24 move to the global scope.
 		const FColor GameThreadColor = FColorList::Red;
 		const FColor RenderThreadColor = FColorList::Blue;
 		const FColor OtherThreadsColor = FColorList::Grey;
@@ -235,7 +235,7 @@ int32 SProfilerMiniView::OnPaint( const FPaintArgs& Args, const FGeometry& Allot
 				&SolidWhiteBrush,
 				MyClippingRect,
 				DrawEffects,
-				// @TODO yrx 2014-04-09 How to get this color from Slate?
+				// #Profiler: 2014-04-09 How to get this color from Slate?
 				FColor(96,96,96)
 			);
 		}
@@ -646,7 +646,8 @@ const int32 SProfilerMiniView::PositionToFrameIndex( const float InPositionX ) c
 
 void SProfilerMiniView::ProcessData()
 {
-//	SCOPE_LOG_TIME_FUNC();
+	static FTotalTimeAndCount TimeAndCount( 0.0f, 0 );
+	//SCOPE_LOG_TIME_FUNC_WITH_GLOBAL( &TimeAndCount );
 
 	AllFrames.Append( RecentlyAddedFrames );
 	RecentlyAddedFrames.Reset();
