Copying //UE4/Dev-Sequencer to //UE4/Dev-Main (Source: //UE4/Dev-Sequencer @ 3178529)

diff --git a/Engine/Source/Developer/Profiler/Private/ProfilerDataSource.h b/Engine/Source/Developer/Profiler/Private/ProfilerDataSource.h
--- a/Engine/Source/Developer/Profiler/Private/ProfilerDataSource.h
+++ b/Engine/Source/Developer/Profiler/Private/ProfilerDataSource.h
@@ -2,9 +2,14 @@
 
 #pragma once
 
-/*-----------------------------------------------------------------------------
-Event graph related type definitions
------------------------------------------------------------------------------*/
+#include "ProfilerSample.h"
+
+
+class FProfilerAggregatedStat;
+class FGraphDataSource;
+class FProfilerSession;
+class IDataProvider;
+
 
 /** Type definition for shared pointers to instances of FEventGraphSample. */
 typedef TSharedPtr<class FEventGraphSample> FEventGraphSamplePtr;
@@ -21,6 +26,7 @@ typedef TSharedPtr<class FEventGraphData, ESPMode::ThreadSafe> FEventGraphDataPt
 /** Type definition for shared references to instances of FEventGraphData. */
 typedef TSharedRef<class FEventGraphData, ESPMode::ThreadSafe> FEventGraphDataRef;
 
+
 /** Scratch buffers for multithreaded usage. */
 struct FProfilerScratchArea : public TThreadSingleton<FProfilerScratchArea>
 {
@@ -470,7 +476,7 @@ protected:
 	 * @param InStatID			- the ID of the stat that this graph data source will be created for
 	 *
 	 */
-	FGraphDataSource( const FProfilerSessionRef& InProfilerSession, const uint32 InStatID  );
+	FGraphDataSource( const TSharedRef<FProfilerSession>& InProfilerSession, const uint32 InStatID  );
 
 public:
 	/** Virtual destructor. */
@@ -497,7 +503,7 @@ public:
 
 	const float GetTotalTimeMS() const;
 
-	const IDataProviderRef GetDataProvider() const;
+	const TSharedRef<IDataProvider> GetDataProvider() const;
 
 	/**
 	 * @return a session instance ID to the profiler session that owns this graph data source. 
@@ -527,7 +533,7 @@ protected:
 	const TGraphDataType GetUncachedValueFromTimeRange( const float StartTimeMS, const float EndTimeMS ) const;
 	
 	/** A reference to the profiler session that owns this graph data source. */
-	const FProfilerSessionRef ProfilerSession;
+	const TSharedRef<FProfilerSession> ProfilerSession;
 
 	// @TODO: This needs to be moved to 'filters and presets' filtering options.
 	float Scale;
@@ -584,7 +590,7 @@ public:
 		return bIsRegistered;
 	}
 
-	void RegisterWithProfilerSession( const FGuid& SessionInstanceID, const FGraphDataSourceRefConst& GraphDataSource )
+	void RegisterWithProfilerSession( const FGuid& SessionInstanceID, const TSharedRef<const FGraphDataSource>& GraphDataSource )
 	{
 		GraphDataSources.Add( SessionInstanceID, GraphDataSource );
 		ThisCachedDataByTime::ClearCache();
@@ -605,7 +611,7 @@ public:
 		return MemoryUsage;
 	}
 
-	const TMap<FGuid,FGraphDataSourceRefConst>::TConstIterator GetSourcesIterator() const
+	const TMap<FGuid,TSharedRef<const FGraphDataSource>>::TConstIterator GetSourcesIterator() const
 	{
 		return GraphDataSources.CreateConstIterator();
 	}
@@ -615,10 +621,10 @@ public:
 		return GraphDataSources.Num();
 	}
 
-	const FGraphDataSourceRefConst* GetFirstSource() const
+	const TSharedRef<const FGraphDataSource>* GetFirstSource() const
 	{
 		// TODO: Add global accessible Null FGraphDataSourceRef, and other instances used in the profiler to avoid returning empty pointer.
-		const FGraphDataSourceRefConst* Result = NULL;
+		const TSharedRef<const FGraphDataSource>* Result = nullptr;
 		if( GetSourcesNum() > 0 )
 		{
 			Result = &GetSourcesIterator().Value();
@@ -646,7 +652,7 @@ public:
 
 			for( auto It = GetSourcesIterator(); It; ++It )
 			{
-				const FGraphDataSourceRefConst& GraphDataSource = It.Value();
+				const TSharedRef<const FGraphDataSource>& GraphDataSource = It.Value();
 				MinTotalTime = FMath::Min( MinTotalTime, GraphDataSource->GetTotalTimeMS() );
 			}
 
@@ -672,7 +678,7 @@ protected:
 	const FVector GetUncachedValueFromTimeRange( const float StartTimeMS, const float EndTimeMS ) const;
 
 	/** A map of graph data sources for all active profiler session instances for the specified stat ID. */
-	TMap<FGuid,FGraphDataSourceRefConst> GraphDataSources;
+	TMap<FGuid,TSharedRef<const FGraphDataSource>> GraphDataSources;
 };
 
 /*-----------------------------------------------------------------------------
@@ -680,90 +686,84 @@ protected:
 -----------------------------------------------------------------------------*/
 
 /** Enumerates event graph columns index. */
-namespace EEventPropertyIndex
+enum EEventPropertyIndex
 {
-	enum Type
-	{
-		/** Stat name must be the first column, because of the expander arrow. */
-		StatName,
-		InclusiveTimeMS,
-		InclusiveTimePct,
-		ExclusiveTimeMS,
-		ExclusiveTimePct,
-		NumCallsPerFrame,
-		/** Special name used for unknown property. */
-		None,
-
-		MinInclusiveTimeMS,
-		MaxInclusiveTimeMS,
-		AvgInclusiveTimeMS,	
-		MinNumCallsPerFrame,
-		MaxNumCallsPerFrame,
-		AvgNumCallsPerFrame,
-		ThreadName,
-		ThreadDurationMS,
-		FrameDurationMS,
-		ThreadPct,
-		FramePct,
-		ThreadToFramePct,
-		GroupName,
+	/** Stat name must be the first column, because of the expander arrow. */
+	StatName,
+	InclusiveTimeMS,
+	InclusiveTimePct,
+	ExclusiveTimeMS,
+	ExclusiveTimePct,
+	NumCallsPerFrame,
+	/** Special name used for unknown property. */
+	None,
+
+	MinInclusiveTimeMS,
+	MaxInclusiveTimeMS,
+	AvgInclusiveTimeMS,	
+	MinNumCallsPerFrame,
+	MaxNumCallsPerFrame,
+	AvgNumCallsPerFrame,
+	ThreadName,
+	ThreadDurationMS,
+	FrameDurationMS,
+	ThreadPct,
+	FramePct,
+	ThreadToFramePct,
+	GroupName,
 		
-		// Booleans
-		bIsHotPath,
-		bIsFiltered,
-		bIsCulled,
+	// Booleans
+	bIsHotPath,
+	bIsFiltered,
+	bIsCulled,
 
-		// Booleans internal
-		bNeedNotCulledChildrenUpdate,
+	// Booleans internal
+	bNeedNotCulledChildrenUpdate,
+
+	/** Invalid enum type, may be used as a number of enumerations. */
+	InvalidOrMax,
+};
 
-		/** Invalid enum type, may be used as a number of enumerations. */
-		InvalidOrMax,
-	};
-}
 
 /** Enumerates event graph sample value formatting types, usually match with event graph widget's columns. */
-namespace EEventPropertyFormatters
+enum class EEventPropertyFormatters
 {
-	enum Type
-	{
-		/** Name, stored as a string, displayed as a regular string. */
-		Name,
+	/** Name, stored as a string, displayed as a regular string. */
+	Name,
 
-		/** Time in milliseconds, stored as a double, displayed as ".3f ms" */
-		TimeMS,
+	/** Time in milliseconds, stored as a double, displayed as ".3f ms" */
+	TimeMS,
 
-		/** Time as percent, stored as a double, displayed as ".1f %" */
-		TimePct,
+	/** Time as percent, stored as a double, displayed as ".1f %" */
+	TimePct,
 
-		/** Number of calls, store as a double, displayed as ".1f" */
-		Number,
+	/** Number of calls, store as a double, displayed as ".1f" */
+	Number,
 
-		/** Boolean value, store as a bool, displaying is not supported yet. */
-		Bool,
+	/** Boolean value, store as a bool, displaying is not supported yet. */
+	Bool,
+
+	/** Invalid enum type, may be used as a number of enumerations. */
+	InvalidOrMax,
+};
 
-		/** Invalid enum type, may be used as a number of enumerations. */
-		InvalidOrMax,
-	};
-}
 
 /** Enumerates . */
-namespace EEventPropertyTypes
+enum class EEventPropertyTypes
 {
-	enum Type
-	{
-		/** double. */
-		Double,
+	/** double. */
+	Double,
 
-		/** FName. */
-		Name,
+	/** FName. */
+	Name,
 
-		/** bool. */
-		Bool,
+	/** bool. */
+	Bool,
+
+	/** Invalid enum type, may be used as a number of enumerations. */
+	InvalidOrMax,
+};
 
-		/** Invalid enum type, may be used as a number of enumerations. */
-		InvalidOrMax,
-	};
-}
 
 //namespace NEventProp
 //{
@@ -773,10 +773,10 @@ namespace EEventPropertyTypes
 	protected:
 		FEventProperty
 		( 
-			const EEventPropertyIndex::Type PropertyIndex, 
+			const EEventPropertyIndex PropertyIndex, 
 			const FName PropertyName, 
 			const uint32 PropertyOffset, 
-			const EEventPropertyFormatters::Type PropertyFormatter
+			const EEventPropertyFormatters PropertyFormatter
 		)
 			:
 			Index( PropertyIndex ),
@@ -786,7 +786,7 @@ namespace EEventPropertyTypes
 			Type( GetTypeFromFormatter(PropertyFormatter) )
 		{}
 
-		EEventPropertyTypes::Type GetTypeFromFormatter( const EEventPropertyFormatters::Type PropertyFormatter ) const
+		EEventPropertyTypes GetTypeFromFormatter( const EEventPropertyFormatters PropertyFormatter ) const
 		{
 			switch( PropertyFormatter )
 			{
@@ -824,11 +824,11 @@ namespace EEventPropertyTypes
 		}
 
 	public:
-		const EEventPropertyIndex::Type Index;
+		const EEventPropertyIndex Index;
 		const FName Name;
 		const uint32 Offset;
-		const EEventPropertyFormatters::Type Formatter;
-		const EEventPropertyTypes::Type Type;
+		const EEventPropertyFormatters Formatter;
+		const EEventPropertyTypes Type;
 	};
 
 	template< typename Type >
@@ -890,7 +890,7 @@ namespace EEventPropertyTypes
 
 namespace NEventFormatter
 {
-	template< EEventPropertyFormatters::Type PropertyType >
+	template< EEventPropertyFormatters PropertyType >
 	FORCEINLINE FString ToString( const FEventGraphSample& Event, const FEventProperty& EventProperty )
 	{
 		check(0);
@@ -952,9 +952,9 @@ public:
 	/** Initializes the minimal property manager for the event graph sample. */
 	static void InitializePropertyManagement();
 
-	static FORCEINLINE_DEBUGGABLE const FEventProperty& GetEventPropertyByIndex( const EEventPropertyIndex::Type PropertyIndex )
+	static FORCEINLINE_DEBUGGABLE const FEventProperty& GetEventPropertyByIndex( const EEventPropertyIndex PropertyIndex )
 	{
-		return Properties[PropertyIndex];
+		return Properties[(const int32)PropertyIndex];
 	}
 
 	static FORCEINLINE_DEBUGGABLE const FEventProperty& GetEventPropertyByName( const FName PropertyName )
@@ -1023,7 +1023,7 @@ protected:
 		const double InInclusiveTimeMS,
 		const double InNumCallsPerFrame,
 
-		const FEventGraphSamplePtr InParentPtr = NULL
+		const FEventGraphSamplePtr InParentPtr = nullptr
 	)
 		: _ParentPtr( InParentPtr )
 		, _RootPtr( nullptr )
@@ -1611,7 +1611,7 @@ public:
 #endif // 0
 
 public:
-	double& PropertyValueAsDouble( const EEventPropertyIndex::Type PropertyIndex )
+	double& PropertyValueAsDouble( const EEventPropertyIndex PropertyIndex )
 	{
 		const FEventProperty& EventProperty = GetEventPropertyByIndex(PropertyIndex);
 		// In case if you want to get FName/bool as a double.
@@ -1620,7 +1620,7 @@ public:
 		return Value.GetPropertyValueRef();
 	}
 
-	FName& PropertyValueAsFName( const EEventPropertyIndex::Type PropertyIndex )
+	FName& PropertyValueAsFName( const EEventPropertyIndex PropertyIndex )
 	{
 		const FEventProperty& EventProperty = GetEventPropertyByIndex(PropertyIndex);
 		// In case if you want to get a double/bool as FName.
@@ -1629,7 +1629,7 @@ public:
 		return Value.GetPropertyValueRef();
 	}
 
-	bool& PropertyValueAsBool( const EEventPropertyIndex::Type PropertyIndex )
+	bool& PropertyValueAsBool( const EEventPropertyIndex PropertyIndex )
 	{
 		const FEventProperty& EventProperty = GetEventPropertyByIndex(PropertyIndex);
 		// In case if you want to get a double/FName as a bool.
@@ -1638,7 +1638,7 @@ public:
 		return Value.GetPropertyValueRef();
 	}
 
-	const FString GetPropertyValueAsString( const EEventPropertyIndex::Type PropertyIndex ) const
+	const FString GetPropertyValueAsString( const EEventPropertyIndex PropertyIndex ) const
 	{
 		const FEventProperty& EventProperty = GetEventPropertyByIndex(PropertyIndex);
 		// In case if you want to get a double/bool as FName.
@@ -1647,7 +1647,7 @@ public:
 		return Value.GetPropertyValueRef().GetPlainNameString();
 	}
 
-	const FString GetFormattedValue( EEventPropertyIndex::Type PropertyIndex ) const
+	const FString GetFormattedValue( EEventPropertyIndex PropertyIndex ) const
 	{
 		const FEventProperty& EventProperty = GetEventPropertyByIndex(PropertyIndex);
 
@@ -1689,7 +1689,7 @@ public:
 	};
 
  public:
-	template< const EEventPropertyIndex::Type BooleanPropertyIndex >
+	template< const EEventPropertyIndex BooleanPropertyIndex >
 	void SetBooleanStateForAllChildren( const bool bValue )
 	{
 		//static_assert(false, "Method not supported.");
@@ -1716,7 +1716,7 @@ public:
 public:
 
 	/** Not used, optimized version @see SetBooleanStateForAllChildren. */
-	void SetBooleanStateForAllChildren_Recurrent( const EEventPropertyIndex::Type BooleanPropertyIndex, const bool bValue )
+	void SetBooleanStateForAllChildren_Recurrent( const EEventPropertyIndex BooleanPropertyIndex, const bool bValue )
 	{
 		PropertyValueAsBool(BooleanPropertyIndex) = bValue;
 		const int32 NumChildren = _ChildrenPtr.Num();
@@ -1850,14 +1850,14 @@ namespace EEventCompareOps
 namespace EventGraphPrivate
 {
 	/** Helper class used to during sorting the events. */
-	template< const EEventPropertyTypes::Type PropType, const EEventCompareOps::Type CompareOp >
+	template< const EEventPropertyTypes PropType, const EEventCompareOps::Type CompareOp >
 	struct TCompareByProperty
 	{
 	private:
-		const EEventPropertyIndex::Type PropertyIndex;
+		const EEventPropertyIndex PropertyIndex;
 
 	public:
-		FORCEINLINE TCompareByProperty( const EEventPropertyIndex::Type InPropertyIndex )
+		FORCEINLINE TCompareByProperty( const EEventPropertyIndex InPropertyIndex )
 			: PropertyIndex( InPropertyIndex )
 		{}
 	};
@@ -1866,10 +1866,10 @@ namespace EventGraphPrivate
 	struct TCompareByProperty<EEventPropertyTypes::Name,EEventCompareOps::Greater>
 	{
 	private:
-		const EEventPropertyIndex::Type PropertyIndex;
+		const EEventPropertyIndex PropertyIndex;
 
 	public:
-		FORCEINLINE TCompareByProperty( const EEventPropertyIndex::Type InPropertyIndex )
+		FORCEINLINE TCompareByProperty( const EEventPropertyIndex InPropertyIndex )
 			: PropertyIndex( InPropertyIndex )
 		{}
 
@@ -1883,10 +1883,10 @@ namespace EventGraphPrivate
 	struct TCompareByProperty<EEventPropertyTypes::Double,EEventCompareOps::Greater>
 	{
 	private:
-		const EEventPropertyIndex::Type PropertyIndex;
+		const EEventPropertyIndex PropertyIndex;
 
 	public:
-		FORCEINLINE TCompareByProperty( const EEventPropertyIndex::Type InPropertyIndex )
+		FORCEINLINE TCompareByProperty( const EEventPropertyIndex InPropertyIndex )
 			: PropertyIndex( InPropertyIndex )
 		{}
 
@@ -1900,10 +1900,10 @@ namespace EventGraphPrivate
 	struct TCompareByProperty<EEventPropertyTypes::Name,EEventCompareOps::Less>
 	{
 	private:
-		const EEventPropertyIndex::Type PropertyIndex;
+		const EEventPropertyIndex PropertyIndex;
 
 	public:
-		FORCEINLINE TCompareByProperty( const EEventPropertyIndex::Type InPropertyIndex )
+		FORCEINLINE TCompareByProperty( const EEventPropertyIndex InPropertyIndex )
 			: PropertyIndex( InPropertyIndex )
 		{}
 
@@ -1917,10 +1917,10 @@ namespace EventGraphPrivate
 	struct TCompareByProperty<EEventPropertyTypes::Double,EEventCompareOps::Less>
 	{
 	private:
-		const EEventPropertyIndex::Type PropertyIndex;
+		const EEventPropertyIndex PropertyIndex;
 
 	public:
-		FORCEINLINE TCompareByProperty( const EEventPropertyIndex::Type InPropertyIndex )
+		FORCEINLINE TCompareByProperty( const EEventPropertyIndex InPropertyIndex )
 			: PropertyIndex( InPropertyIndex )
 		{}
 
@@ -1946,9 +1946,9 @@ public:
 	static void ExecuteOperation
 	( 
 		FEventGraphSamplePtr DestPtr, 
-		const EEventPropertyIndex::Type DestPropertyIndex, 
+		const EEventPropertyIndex DestPropertyIndex, 
 		const FEventGraphSamplePtr SrcPtr, 
-		const EEventPropertyIndex::Type SrcPropertyIndex, 
+		const EEventPropertyIndex SrcPropertyIndex, 
 		const EEventCompareOps::Type OpType
 	)
 	{
@@ -1986,7 +1986,7 @@ public:
 	}
 
 private:
-	template < const EEventPropertyTypes::Type PropType, const EEventCompareOps::Type CompareOp >
+	template < const EEventPropertyTypes PropType, const EEventCompareOps::Type CompareOp >
 	static void ExecuteOperationInternal
 	( 
 		TArray< FEventGraphSamplePtr >& DestEvents, 
@@ -2050,7 +2050,7 @@ public:
 	}
 
 private:
-	template< const EEventPropertyTypes::Type PropType, const EEventCompareOps::Type CompareOp >
+	template< const EEventPropertyTypes PropType, const EEventCompareOps::Type CompareOp >
 	static void SortInternal( TArray< FEventGraphSamplePtr >& ChildrenToSort, const EventGraphPrivate::TCompareByProperty<PropType,CompareOp>& CompareInstance )
 	{
 		ChildrenToSort.Sort( CompareInstance );
@@ -2170,7 +2170,7 @@ protected:
 		const FProfilerSession * const ProfilerSession,
 		const FEventGraphSamplePtr ParentEvent, 
 		const FProfilerSample& ParentSample, 
-		const IDataProviderRef DataProvider
+		const TSharedRef<IDataProvider> DataProvider
 	);
 
 public:
diff --git a/Engine/Source/Developer/Profiler/Private/ProfilerFPSAnalyzer.cpp b/Engine/Source/Developer/Profiler/Private/ProfilerFPSAnalyzer.cpp
--- a/Engine/Source/Developer/Profiler/Private/ProfilerFPSAnalyzer.cpp
+++ b/Engine/Source/Developer/Profiler/Private/ProfilerFPSAnalyzer.cpp
@@ -1,10 +1,8 @@
 // Copyright 1998-2016 Epic Games, Inc. All Rights Reserved.
 
 #include "ProfilerPrivatePCH.h"
+#include "ProfilerFPSAnalyzer.h"
 
-/*-----------------------------------------------------------------------------
-	FFPSAnalyzer
------------------------------------------------------------------------------*/
 
 void FFPSAnalyzer::Reset()
 {
@@ -21,31 +19,38 @@ void FFPSAnalyzer::Reset()
 	FPS20 = 0;
 }
 
-void FFPSAnalyzer::AddSample( float FPSSample )
+
+void FFPSAnalyzer::AddSample(float FPSSample)
 {
 	Samples.Add(FPSSample);
 
 	float DeltaSeconds = 1.0f / FPSSample;
 	int32 Index = FPlatformMath::FloorToInt(FPSSample / Interval);
 	Index = Index >= Histogram.Num() ? Histogram.Num()-1 : Index;
+	
 	Histogram[Index].Count++;
 	Histogram[Index].CummulativeTime += DeltaSeconds;
+	
 	if (FPSSample >= 90)
 	{
 		FPS90++;
 	}
+	
 	if (FPSSample >= 60)
 	{
 		FPS60++;
 	}
+	
 	if (FPSSample >= 30)
 	{
 		FPS30++;
 	}
+	
 	if (FPSSample >= 25)
 	{
 		FPS25++;
 	}
+	
 	if (FPSSample >= 20)
 	{
 		FPS20++;
@@ -55,27 +60,33 @@ void FFPSAnalyzer::AddSample( float FPSSample )
 	{
 		MaxFPS = FPSSample;
 	}
+	
 	if (FPSSample < MinFPS)
 	{
 		MinFPS = FPSSample;
 	}
+	
 	AveFPS = (FPSSample + (float)(Samples.Num()-1) * AveFPS) / (float)Samples.Num();
 }
 
+
 SIZE_T FFPSAnalyzer::GetMemoryUsage() const
 {
 	const SIZE_T MemoryAllocated = Samples.GetAllocatedSize() + Histogram.GetAllocatedSize();
 	return MemoryAllocated;
 }
 
+
 int32 FFPSAnalyzer::GetTotalCount()
 {
 	return Samples.Num();
 }
 
+
 int32 FFPSAnalyzer::GetCount(float InMinVal, float InMaxVal)
 {
 	int32 Index = FPlatformMath::FloorToInt(InMinVal / Interval);
 	Index = Index >= Histogram.Num() ? Histogram.Num()-1 : Index;
+
 	return Histogram[Index].Count;
-}
\ No newline at end of file
+}
diff --git a/Engine/Source/Developer/Profiler/Private/ProfilerFPSAnalyzer.h b/Engine/Source/Developer/Profiler/Private/ProfilerFPSAnalyzer.h
--- a/Engine/Source/Developer/Profiler/Private/ProfilerFPSAnalyzer.h
+++ b/Engine/Source/Developer/Profiler/Private/ProfilerFPSAnalyzer.h
@@ -2,9 +2,8 @@
 
 #pragma once
 
-/*-----------------------------------------------------------------------------
-	Declarations
------------------------------------------------------------------------------*/
+#include "ProfilerSample.h"
+
 
 struct FProfilerFPSChartEntry
 {
@@ -12,6 +11,7 @@ struct FProfilerFPSChartEntry
 	double CummulativeTime;
 };
 
+
 enum FPSChartBins
 {
 	FPSBin_0_5,
@@ -36,10 +36,14 @@ enum FPSChartBins
 	FPSBin_LastBucketStat,
 
 };
+
+
 /** Implements a frame rate analyzer */
-class FFPSAnalyzer : public IHistogramDataSource
+class FFPSAnalyzer
+	: public IHistogramDataSource
 {
 public:
+
 	FFPSAnalyzer(int32 InInterval, int32 InMinVal, int32 InMaxVal)
 	{
 		Interval = InInterval;
@@ -48,19 +52,16 @@ public:
 		Reset();
 	}
 
-	virtual ~FFPSAnalyzer()
-	{}
-
+	virtual ~FFPSAnalyzer() { }
 	virtual int32 GetCount(float MinVal, float MaxVal) override;
-
 	virtual int32 GetTotalCount() override;
 
 	void Reset();
-
-	void AddSample( float FPSSample );
-
+	void AddSample(float FPSSample);
 	SIZE_T GetMemoryUsage() const;
 
+public:
+
 	TArray<float> Samples;
 	TArray<FProfilerFPSChartEntry> Histogram;
 	float MinFPS;
@@ -74,4 +75,4 @@ public:
 	int32 Interval;
 	int32 MaxVal;
 	int32 MinVal;
-};
\ No newline at end of file
+};
diff --git a/Engine/Source/Developer/Profiler/Private/ProfilerManager.cpp b/Engine/Source/Developer/Profiler/Private/ProfilerManager.cpp
--- a/Engine/Source/Developer/Profiler/Private/ProfilerManager.cpp
+++ b/Engine/Source/Developer/Profiler/Private/ProfilerManager.cpp
@@ -1,20 +1,20 @@
 // Copyright 1998-2016 Epic Games, Inc. All Rights Reserved.
 
 #include "ProfilerPrivatePCH.h"
-#include "RenderingThread.h"
+#include "ProfilerCommands.h"
+#include "ProfilerDataProvider.h"
+#include "ProfilerDataSource.h"
+#include "ProfilerManager.h"
+#include "ProfilerSession.h"
+#include "SProfilerMiniView.h"
+#include "SProfilerWindow.h"
+
 
 #define LOCTEXT_NAMESPACE "FProfilerCommands"
 
-/*-----------------------------------------------------------------------------
-	Globals and statics
------------------------------------------------------------------------------*/
 DEFINE_LOG_CATEGORY(Profiler);
-FProfilerSettings FProfilerSettings::Defaults( true );
-
-/*-----------------------------------------------------------------------------
-	Stat declarations.
------------------------------------------------------------------------------*/
 
+FProfilerSettings FProfilerSettings::Defaults( true );
 
 
 DEFINE_STAT(STAT_DG_OnPaint);
@@ -22,9 +22,6 @@ DEFINE_STAT(STAT_PM_HandleProfilerData);
 DEFINE_STAT(STAT_PM_Tick);
 DEFINE_STAT(STAT_PM_MemoryUsage);
 
-/*-----------------------------------------------------------------------------
-	Profiler manager implementation
------------------------------------------------------------------------------*/
 
 TSharedPtr<FProfilerManager> FProfilerManager::Instance = nullptr;
 
@@ -36,7 +33,8 @@ struct FEventGraphSampleLess
 	}
 };
 
-FProfilerManager::FProfilerManager( const ISessionManagerRef& InSessionManager )
+
+FProfilerManager::FProfilerManager( const TSharedRef<ISessionManager>& InSessionManager )
 	: SessionManager( InSessionManager )
 	, CommandList( new FUICommandList() )
 	, ProfilerActionManager( this )
@@ -168,6 +166,43 @@ FProfilerManager::~FProfilerManager()
 	ClearStatsAndInstances();
 }
 
+
+TSharedPtr<FProfilerManager> FProfilerManager::Get()
+{
+	return FProfilerManager::Instance;
+}
+
+
+FProfilerActionManager& FProfilerManager::GetActionManager()
+{
+	return FProfilerManager::Instance->ProfilerActionManager;
+}
+
+
+FProfilerSettings& FProfilerManager::GetSettings()
+{
+	return FProfilerManager::Instance->Settings;
+}
+
+
+const TSharedRef< FUICommandList > FProfilerManager::GetCommandList() const
+{
+	return CommandList;
+}
+
+
+const FProfilerCommands& FProfilerManager::GetCommands()
+{
+	return FProfilerCommands::Get();
+}
+
+
+TSharedPtr<FProfilerSession> FProfilerManager::GetProfilerSession()
+{
+	return ProfilerSession;
+}
+
+
 class FStatsHeaderReader : public FStatsReadFile
 {
 	friend struct FStatsReader<FStatsHeaderReader>;
@@ -405,7 +440,8 @@ void FProfilerManager::ProfilerClient_OnProfilerFileTransfer( const FString& Fil
 
 	const float Progress = (double)FileProgress/(double)FileSize;
 
-	ELoadingProgressStates::Type ProgressState = ELoadingProgressStates::InvalidOrMax;
+	ELoadingProgressStates ProgressState = ELoadingProgressStates::InvalidOrMax;
+
 	if( FileProgress == 0 )
 	{
 		ProgressState = ELoadingProgressStates::Started;
@@ -443,7 +479,7 @@ void FProfilerManager::ProfilerClient_OnClientDisconnected( const FGuid& Session
 
 void FProfilerManager::SessionManager_OnInstanceSelectionChanged(const TSharedPtr<ISessionInstanceInfo>& InInstance, bool Selected)
 {
-	const ISessionInfoPtr& SelectedSession = SessionManager->GetSelectedSession();
+	const TSharedPtr<ISessionInfo>& SelectedSession = SessionManager->GetSelectedSession();
 	const bool SessionIsValid = SelectedSession.IsValid()
 		&& (SelectedSession->GetSessionOwner() == FPlatformProcess::UserName(false))
 		&& (SessionManager->GetSelectedInstances().Num() > 0);
@@ -495,7 +531,7 @@ bool FProfilerManager::TrackStat( const uint32 StatID )
 	const bool bStatIsReady = ProfilerSession->GetAggregatedStat( StatID ) != nullptr;
 	if( StatID != 0 && bStatIsReady )
 	{
-		FTrackedStatPtr TrackedStat = TrackedStats.FindRef( StatID );
+		TSharedPtr<FTrackedStat> TrackedStat = TrackedStats.FindRef( StatID );
 		if (!TrackedStat.IsValid())
 		{
 			// R = H, G = S, B = V
@@ -530,7 +566,7 @@ bool FProfilerManager::UntrackStat( const uint32 StatID )
 	const uint32 GameThreadStatID = ProfilerSession->GetMetaData()->GetGameThreadStatID();
 	if (StatID != GameThreadStatID)
 	{
-		FTrackedStatPtr TrackedStat = TrackedStats.FindRef( StatID );
+		TSharedPtr<FTrackedStat> TrackedStat = TrackedStats.FindRef( StatID );
 		if (TrackedStat.IsValid())
 		{
 			TrackedStatChangedEvent.Broadcast( TrackedStat, false );
@@ -556,7 +592,7 @@ void FProfilerManager::ClearStatsAndInstances()
 
 	for( auto It = TrackedStats.CreateConstIterator(); It; ++It )
 	{
-		const FTrackedStatPtr& TrackedStat = It.Value();
+		const TSharedPtr<FTrackedStat>& TrackedStat = It.Value();
 		TrackedStatChangedEvent.Broadcast( TrackedStat, false );
 	}
 	TrackedStats.Empty();
@@ -628,7 +664,7 @@ void FProfilerManager::DataGraph_OnSelectionChangedForIndex( uint32 FrameStartIn
 	}
 }
 
-void FProfilerManager::ProfilerSession_OnAddThreadTime( int32 FrameIndex, const TMap<uint32, float>& ThreadMS, const FProfilerStatMetaDataRef& StatMetaData )
+void FProfilerManager::ProfilerSession_OnAddThreadTime( int32 FrameIndex, const TMap<uint32, float>& ThreadMS, const TSharedRef<FProfilerStatMetaData>& StatMetaData )
 {
 	TSharedPtr<SProfilerWindow> ProfilerWindowPtr = GetProfilerWindow();
 	if( ProfilerWindowPtr.IsValid() )
@@ -641,7 +677,7 @@ void FProfilerManager::ProfilerSession_OnAddThreadTime( int32 FrameIndex, const
 	}
 }
 
-void FProfilerManager::SetViewMode( EProfilerViewMode::Type NewViewMode )
+void FProfilerManager::SetViewMode( EProfilerViewMode NewViewMode )
 {
 	if( NewViewMode != ViewMode )
 	{
diff --git a/Engine/Source/Developer/Profiler/Private/ProfilerManager.h b/Engine/Source/Developer/Profiler/Private/ProfilerManager.h
--- a/Engine/Source/Developer/Profiler/Private/ProfilerManager.h
+++ b/Engine/Source/Developer/Profiler/Private/ProfilerManager.h
@@ -2,20 +2,29 @@
  
 #pragma once
 
-/*-----------------------------------------------------------------------------
-	Basic structures
------------------------------------------------------------------------------*/
+#include "ProfilerCommands.h"
+
+
+enum class EProfilerSessionTypes;
+enum class EProfilerViewMode;
+class FGraphDataSource;
+class FProfilerCommands;
+class FProfilerSession;
+class FProfilerStatMetaData;
+class IProfilerClient;
+class SProfilerWindow;
+
 
 /** Contains all settings for the profiler, accessible through the profiler manager. */
 class FProfilerSettings
 {
 public:
-	FProfilerSettings( bool bInIsDefault = false )
-		: bShowCoalescedViewModesInEventGraph( true )
-		, bIsEditing( false )
-		, bIsDefault( bInIsDefault )
+	FProfilerSettings(bool bInIsDefault = false)
+		: bShowCoalescedViewModesInEventGraph(true)
+		, bIsEditing(false)
+		, bIsDefault(bInIsDefault)
 	{
-		if( !bIsDefault )
+		if(!bIsDefault)
 		{
 			LoadFromConfig();
 		}
@@ -23,7 +32,7 @@ public:
 
 	~FProfilerSettings()
 	{
-		if( !bIsDefault )
+		if(!bIsDefault)
 		{
 			SaveToConfig();
 		}
@@ -33,13 +42,13 @@ public:
 	{
 		FConfigCacheIni::LoadGlobalIniFile(ProfilerSettingsIni, TEXT("ProfilerSettings"));
 
-		GConfig->GetBool( TEXT("Profiler.ProfilerOptions"), TEXT("bShowCoalescedViewModesInEventGraph"), bShowCoalescedViewModesInEventGraph, ProfilerSettingsIni );
+		GConfig->GetBool(TEXT("Profiler.ProfilerOptions"), TEXT("bShowCoalescedViewModesInEventGraph"), bShowCoalescedViewModesInEventGraph, ProfilerSettingsIni);
 	}
 
 	void SaveToConfig()
 	{
-		GConfig->SetBool( TEXT("Profiler.ProfilerOptions"), TEXT("bShowCoalescedViewModesInEventGraph"), bShowCoalescedViewModesInEventGraph, ProfilerSettingsIni );
-		GConfig->Flush( false, ProfilerSettingsIni );
+		GConfig->SetBool(TEXT("Profiler.ProfilerOptions"), TEXT("bShowCoalescedViewModesInEventGraph"), bShowCoalescedViewModesInEventGraph, ProfilerSettingsIni);
+		GConfig->Flush(false, ProfilerSettingsIni);
 	}
 
 	void EnterEditMode()
@@ -63,7 +72,8 @@ public:
 	}
 
 //protected:
-	/** Contains default settings. */
+
+/** Contains default settings. */
 	static FProfilerSettings Defaults;
 
 	/** Profiler setting filename ini. */
@@ -81,35 +91,38 @@ public:
 
 
 /** Contains basic information about tracked stat. */
-class FTrackedStat : public FNoncopyable, public TSharedFromThis<FTrackedStat>
+class FTrackedStat
+	: public FNoncopyable
+	, public TSharedFromThis<FTrackedStat>
 {
 public:
+
 	/**
 	 * Initialization constructor.
 	 *
-	 * @param InGraphDataSource - Data source
-	 * @param InGraphColor - color used to draw the average value
-	 * @param InStatID - the ID of the stat which will be tracked
+	 * @param InGraphDataSource Data source
+	 * @param InGraphColor Color used to draw the average value
+	 * @param InStatID The ID of the stat which will be tracked
 	 *
 	 */
 	FTrackedStat
 	(
-		FGraphDataSourceRefConst InGraphDataSource,
+		TSharedRef<const FGraphDataSource> InGraphDataSource,
 		const FLinearColor InGraphColor,
 		const uint32 InStatID	
 	)
-		: GraphDataSource( InGraphDataSource )
-		, GraphColor( InGraphColor )
-		, StatID( InStatID )
-	{}
+		: GraphDataSource(InGraphDataSource)
+		, GraphColor(InGraphColor)
+		, StatID(InStatID)
+	{ }
 
 	/** Destructor. */
-	~FTrackedStat()
-	{}
+	~FTrackedStat() { }
 
 //protected:
+
 	/** A shared reference to the graph data source for active profiler session for the specified stat ID. */
-	FGraphDataSourceRefConst GraphDataSource;
+	TSharedRef<const FGraphDataSource> GraphDataSource;
 
 	/** A color to visualize graph value for the data graph. */
 	const FLinearColor GraphColor;
@@ -118,22 +131,17 @@ public:
 	const uint32 StatID;
 };
 
-/*-----------------------------------------------------------------------------
-	FProfilerManager
------------------------------------------------------------------------------*/
 
-namespace EProfilerViewMode
+enum class EProfilerViewMode
 {
-	enum Type
-	{
-		/** Regular line graphs, for the regular stats file. */
-		LineIndexBased,
-		/** Thread view graph, for the raw stats file. */
-		ThreadViewTimeBased,
-		/** Invalid enum type, may be used as a number of enumerations. */
-		InvalidOrMax,
-	};
-}
+	/** Regular line graphs, for the regular stats file. */
+	LineIndexBased,
+	/** Thread view graph, for the raw stats file. */
+	ThreadViewTimeBased,
+	/** Invalid enum type, may be used as a number of enumerations. */
+	InvalidOrMax,
+};
+
 
 /** 
  ** This class manages following areas:
@@ -153,10 +161,10 @@ public:
 	/**
 	 * Creates a profiler manager, only one instance can exist.
 	 *
-	 * @param InSessionManager	- the session manager to use
+	 * @param InSessionManager The session manager to use
 	 *
 	 */
-	FProfilerManager( const ISessionManagerRef& InSessionManager );
+	FProfilerManager(const TSharedRef<ISessionManager>& InSessionManager);
 
 	/** Virtual destructor. */
 	virtual ~FProfilerManager();
@@ -164,21 +172,22 @@ public:
 	/**
 	 * Creates an instance of the profiler manager and initializes global instance with the previously created instance of the profiler manager.
 	 *
-	 * @param InSessionManager	- the session manager to use
-	 *
+	 * @param InSessionManager The session manager to use
 	 */
-	static TSharedPtr<FProfilerManager> Initialize( const ISessionManagerRef& SessionManager )
+	static TSharedPtr<FProfilerManager> Initialize(const TSharedRef<ISessionManager>& SessionManager)
 	{
 		if (FProfilerManager::Instance.IsValid())
 		{
 			FProfilerManager::Instance.Reset();
 		}
-		FProfilerManager::Instance = MakeShareable( new FProfilerManager(SessionManager) );
+
+		FProfilerManager::Instance = MakeShareable(new FProfilerManager(SessionManager));
 		FProfilerManager::Instance->PostConstructor();
+
 		return FProfilerManager::Instance;
 	}
 
-	void AssignProfilerWindow( const TSharedRef<class SProfilerWindow>& InProfilerWindow )
+	void AssignProfilerWindow(const TSharedRef<SProfilerWindow>& InProfilerWindow)
 	{
 		ProfilerWindow = InProfilerWindow;
 	}
@@ -206,62 +215,44 @@ public:
 	 *		IProfilerModule& ProfilerModule = FModuleManager::Get().LoadModuleChecked<IProfilerModule>("Profiler");
 	 *		ProfilerModule.GetProfilerManager();
 	 */
-	static TSharedPtr<FProfilerManager> Get()
-	{
-		return FProfilerManager::Instance;
-	}
+	static TSharedPtr<FProfilerManager> Get();
 
 	/**
 	 * @return an instance of the profiler action manager.
 	 */
-	static FProfilerActionManager& GetActionManager()
-	{
-		return FProfilerManager::Instance->ProfilerActionManager;
-	}
+	static FProfilerActionManager& GetActionManager();
 
 	/**
 	 * @return an instance of the profiler settings.
 	 */
-	static FProfilerSettings& GetSettings()
-	{
-		return FProfilerManager::Instance->Settings;
-	}
+	static FProfilerSettings& GetSettings();
 
 	/** @returns UI command list for the profiler manager. */
-	const TSharedRef< FUICommandList > GetCommandList() const
-	{
-		return CommandList;
-	}
+	const TSharedRef< FUICommandList > GetCommandList() const;
 
 	/**
 	 * @return an instance of the profiler commands.
 	 */
-	static const FProfilerCommands& GetCommands()
-	{
-		return FProfilerCommands::Get();
-	}
+	static const FProfilerCommands& GetCommands();
 
 	/**
 	 * @return an instance of the profiler session.
 	 */
-	FProfilerSessionPtr GetProfilerSession()
-	{
-		return ProfilerSession;
-	}
+	TSharedPtr<FProfilerSession> GetProfilerSession();
 
 	/*-----------------------------------------------------------------------------
 		Stat tracking, Session instance management
 	-----------------------------------------------------------------------------*/
 
-	bool TrackStat( const uint32 StatID );
-	bool UntrackStat( const uint32 StatID );
+	bool TrackStat(const uint32 StatID);
+	bool UntrackStat(const uint32 StatID);
 	void TrackDefaultStats();
 	void ClearStatsAndInstances();
 
 	/**
 	 * @return true, if the specified stat is currently tracked by the profiler.
 	 */
-	const bool IsStatTracked( const uint32 StatID ) const;
+	const bool IsStatTracked(const uint32 StatID) const;
 
 	/**
 	 * @return True, if the profiler has at least one fully processed capture file
@@ -303,7 +294,7 @@ public:
 	 * @param bRequestedDataPreviewState - data preview state that should be set
 	 *
 	 */
-	void SetDataPreview( const bool bRequestedDataPreviewState );
+	void SetDataPreview(const bool bRequestedDataPreviewState);
 
 	/** @return true, if all sessions instances are capturing data to a file, only valid if profiler is connected to network based session */
 	const bool IsDataCapturing() const;
@@ -314,30 +305,30 @@ public:
 	 * @param bRequestedDataCaptureState - data capture state that should be set
 	 *
 	 */
-	void SetDataCapture( const bool bRequestedDataCaptureState );
+	void SetDataCapture(const bool bRequestedDataCaptureState);
 
-	void ProfilerSession_OnCaptureFileProcessed( const FGuid ProfilerInstanceID );
-	void ProfilerSession_OnAddThreadTime( int32 FrameIndex, const TMap<uint32, float>& ThreadMS, const FProfilerStatMetaDataRef& StatMetaData );
+	void ProfilerSession_OnCaptureFileProcessed(const FGuid ProfilerInstanceID);
+	void ProfilerSession_OnAddThreadTime(int32 FrameIndex, const TMap<uint32, float>& ThreadMS, const TSharedRef<FProfilerStatMetaData>& StatMetaData);
 
 	/*-----------------------------------------------------------------------------
 		Event graphs management
 	-----------------------------------------------------------------------------*/
 
-	void CreateEventGraphTab( const FGuid ProfilerInstanceID );
+	void CreateEventGraphTab(const FGuid ProfilerInstanceID);
 	void CloseAllEventGraphTabs();
 
 	/*-----------------------------------------------------------------------------
 		Data graphs management
 	-----------------------------------------------------------------------------*/
 
-	void DataGraph_OnSelectionChangedForIndex( uint32 FrameStartIndex, uint32 FrameEndIndex );
+	void DataGraph_OnSelectionChangedForIndex(uint32 FrameStartIndex, uint32 FrameEndIndex);
 
 	/*-----------------------------------------------------------------------------
 		Events declarations
 	-----------------------------------------------------------------------------*/
 	
 public:
-	DECLARE_EVENT_OneParam( FProfilerManager, FViewModeChangedEvent, EProfilerViewMode::Type /*NewViewMode*/ );
+	DECLARE_EVENT_OneParam(FProfilerManager, FViewModeChangedEvent, EProfilerViewMode /*NewViewMode*/);
 	FViewModeChangedEvent& OnViewModeChanged()
 	{
 		return OnViewModeChangedEvent;
@@ -349,13 +340,8 @@ protected:
 	
 
 public:
-	/**
-	 * The event to execute when the status of specified tracked stat has changed.
-	 *
-	 * @param const FTrackedStatPtr& - a reference to the tracked stat whose status has changed, this reference is valid only within the scope of function
-	 * @param bool bIsTracked - true, if stat has been added for tracking, false, if stat has been removed from tracking
-	 */
-	DECLARE_EVENT_TwoParams( FProfilerManager, FTrackedStatChangedEvent, const FTrackedStatPtr&, bool );
+	/** The event to execute when the status of specified tracked stat has changed. */
+	DECLARE_EVENT_TwoParams(FProfilerManager, FTrackedStatChangedEvent, const TSharedPtr<FTrackedStat>& /*TrackedStat*/, bool /*IsTracked*/);
 	FTrackedStatChangedEvent& OnTrackedStatChanged()
 	{
 		return TrackedStatChangedEvent;
@@ -371,9 +357,9 @@ public:
 	/**
 	 * The event to execute when a new frame has been added to the specified profiler session instance.
 	 *
-	 * @param const FProfilerSessionPtr& - a shared pointer to the profiler session instance which received a new frame
+	 * @param const TSharedPtr<FProfilerSession>& - a shared pointer to the profiler session instance which received a new frame
 	 */
-	DECLARE_EVENT_OneParam( FProfilerManager, FFrameAddedEvent, const FProfilerSessionPtr& );
+	DECLARE_EVENT_OneParam(FProfilerManager, FFrameAddedEvent, const TSharedPtr<FProfilerSession>&);
 	FFrameAddedEvent& OnFrameAdded()
 	{
 		return FrameAddedEvent;
@@ -383,7 +369,7 @@ public:
 	/**
 	 * The event to be invoked once per second, for example to update the data graph.
 	 */
-	DECLARE_EVENT( FProfilerManager, FOneSecondPassedEvent );
+	DECLARE_EVENT(FProfilerManager, FOneSecondPassedEvent);
 	FOneSecondPassedEvent& OnOneSecondPassed()
 	{
 		return OneSecondPassedEvent;
@@ -398,7 +384,7 @@ public:
 	/**
 	 * The event to execute when the list of session instances have changed.
 	 */
-	DECLARE_EVENT( FProfilerManager, FOnSessionsUpdatedEvent );
+	DECLARE_EVENT(FProfilerManager, FOnSessionsUpdatedEvent);
 	FOnSessionsUpdatedEvent& OnSessionInstancesUpdated()
 	{
 		return SessionInstancesUpdatedEvent;
@@ -413,7 +399,7 @@ public:
 	/**
 	 * The event to execute when the filter and presets widget should be updated with the latest data.
 	 */
-	DECLARE_EVENT( FProfilerManager, FRequestFilterAndPresetsUpdateEvent );
+	DECLARE_EVENT(FProfilerManager, FRequestFilterAndPresetsUpdateEvent);
 	FRequestFilterAndPresetsUpdateEvent& OnRequestFilterAndPresetsUpdate()
 	{
 		return RequestFilterAndPresetsUpdateEvent;
@@ -430,30 +416,28 @@ public:
 	 *
 	 * @param ProfilerCaptureFilepath	- The path to the file containing a captured session instance
 	 */
-	void LoadProfilerCapture( const FString& ProfilerCaptureFilepath );
+	void LoadProfilerCapture(const FString& ProfilerCaptureFilepath);
 
 	/** Creates a new profiler session instance and load a raw stats file from the specified location. */
-	void LoadRawStatsFile( const FString& RawStatsFileFileath );
+	void LoadRawStatsFile(const FString& RawStatsFileFileath);
 
 protected:
-	void ProfilerClient_OnProfilerData( const FGuid& InstanceID, const FProfilerDataFrame& Content );
-	void ProfilerClient_OnClientConnected( const FGuid& SessioID, const FGuid& InstanceID );
-	void ProfilerClient_OnClientDisconnected( const FGuid& SessionID, const FGuid& InstanceID );
-	void ProfilerClient_OnMetaDataUpdated( const FGuid& InstanceID, const FStatMetaData& MetaData );
-	void ProfilerClient_OnLoadCompleted( const FGuid& InstanceID );
-	void ProfilerClient_OnLoadCancelled( const FGuid& InstanceID );
-	void ProfilerClient_OnLoadStarted( const FGuid& InstanceID );
-
-	void ProfilerClient_OnProfilerFileTransfer( const FString& Filename, int64 FileProgress, int64 FileSize );
-
-	void SessionManager_OnInstanceSelectionChanged( const TSharedPtr<ISessionInstanceInfo>& Instance, bool Selected );
+	void ProfilerClient_OnProfilerData(const FGuid& InstanceID, const FProfilerDataFrame& Content);
+	void ProfilerClient_OnClientConnected(const FGuid& SessioID, const FGuid& InstanceID);
+	void ProfilerClient_OnClientDisconnected(const FGuid& SessionID, const FGuid& InstanceID);
+	void ProfilerClient_OnMetaDataUpdated(const FGuid& InstanceID, const FStatMetaData& MetaData);
+	void ProfilerClient_OnLoadCompleted(const FGuid& InstanceID);
+	void ProfilerClient_OnLoadCancelled(const FGuid& InstanceID);
+	void ProfilerClient_OnLoadStarted(const FGuid& InstanceID);
+	void ProfilerClient_OnProfilerFileTransfer(const FString& Filename, int64 FileProgress, int64 FileSize);
+	void SessionManager_OnInstanceSelectionChanged(const TSharedPtr<ISessionInstanceInfo>& Instance, bool Selected);
 
 public:
-	const FLinearColor& GetColorForStatID( const uint32 StatID ) const;
+	const FLinearColor& GetColorForStatID(const uint32 StatID) const;
 
 protected:
 	/** Updates this manager, done through FCoreTicker. */
-	bool Tick( float DeltaTime );
+	bool Tick(float DeltaTime);
 
 public:
 
@@ -467,7 +451,7 @@ public:
 	}
 
 	/** Sets a new view mode for the profiler. */
-	void SetViewMode( EProfilerViewMode::Type NewViewMode );
+	void SetViewMode(EProfilerViewMode NewViewMode);
 
 protected:
 	/** The delegate to be invoked when this profiler manager ticks. */
@@ -480,19 +464,19 @@ protected:
 	TWeakPtr<class SProfilerWindow> ProfilerWindow;
 
 	/** A shared pointer to the session manager. */
-	ISessionManagerPtr SessionManager;
+	TSharedPtr<ISessionManager> SessionManager;
 
 	/** A shared pointer to the currently selected session in the session browser. */
-	ISessionInfoPtr ActiveSession;
+	TSharedPtr<ISessionInfo> ActiveSession;
 
 	/** A shared pointer to the currently selected instance in the session browser. */
-	FGuid/*ISessionInstanceInfoPtr*/ ActiveInstanceID;
+	FGuid/*TSharedPtr<ISessionInstanceInfo>*/ ActiveInstanceID;
 
 	/** Profiler session, to be removed from here. */
-	FProfilerSessionPtr ProfilerSession;
+	TSharedPtr<FProfilerSession> ProfilerSession;
 
 	/** A shared pointer to the profiler client, which is used to deliver all profiler data from the active session. */
-	IProfilerClientPtr ProfilerClient;
+	TSharedPtr<IProfilerClient> ProfilerClient;
 
 	/** List of UI commands for the profiler manager. This will be filled by this and corresponding classes. */
 	TSharedRef< FUICommandList > CommandList;
@@ -514,7 +498,7 @@ protected:
 	FFrameAddedEvent FrameAddedEvent;
 
 	/** Contains all currently tracked stats, stored as StatID -> FTrackedStat. */
-	TMap<uint32, FTrackedStatPtr> TrackedStats;
+	TMap<uint32, TSharedPtr<FTrackedStat>> TrackedStats;
 
 	
 	/*-----------------------------------------------------------------------------
@@ -522,10 +506,10 @@ protected:
 	-----------------------------------------------------------------------------*/
 
 	/** Profiler session type that is currently initialized. */
-	EProfilerSessionTypes::Type ProfilerType;
+	EProfilerSessionTypes ProfilerType;
 
 	/** Profiler view mode. */
-	EProfilerViewMode::Type ViewMode;
+	EProfilerViewMode ViewMode;
 
 	// TODO: Bool should be replaces with type similar to ECheckBoxState {Checked,Unchecked,Undertermined}
 
diff --git a/Engine/Source/Developer/Profiler/Private/ProfilerModule.cpp b/Engine/Source/Developer/Profiler/Private/ProfilerModule.cpp
--- a/Engine/Source/Developer/Profiler/Private/ProfilerModule.cpp
+++ b/Engine/Source/Developer/Profiler/Private/ProfilerModule.cpp
@@ -1,7 +1,10 @@
 // Copyright 1998-2016 Epic Games, Inc. All Rights Reserved.
 
 #include "ProfilerPrivatePCH.h"
-#include "SDockTab.h"
+#include "ProfilerManager.h"
+#include "ProfilerRawStatsForMemory.h"
+#include "SProfilerWindow.h"
+
 
 /**
  * Implements the FProfilerModule module.
@@ -10,11 +13,10 @@ class FProfilerModule
 	: public IProfilerModule
 {
 public:
-	virtual TSharedRef<SWidget> CreateProfilerWindow( const ISessionManagerRef& InSessionManager, const TSharedRef<SDockTab>& ConstructUnderMajorTab ) override;
 
-	virtual void StartupModule() override
-	{}
+	virtual TSharedRef<SWidget> CreateProfilerWindow( const TSharedRef<ISessionManager>& InSessionManager, const TSharedRef<SDockTab>& ConstructUnderMajorTab ) override;
 
+	virtual void StartupModule() override { }
 	virtual void ShutdownModule() override;
 
 	virtual bool SupportsDynamicReloading() override
@@ -27,18 +29,20 @@ public:
 	FRawStatsMemoryProfiler* OpenRawStatsForMemoryProfiling( const TCHAR* Filename );
 
 protected:
+
 	/** Shutdowns the profiler manager. */
 	void Shutdown( TSharedRef<SDockTab> TabBeingClosed );
 };
 
+
 IMPLEMENT_MODULE( FProfilerModule, Profiler );
 
+
 /*-----------------------------------------------------------------------------
 	FProfilerModule
 -----------------------------------------------------------------------------*/
 
-
-TSharedRef<SWidget> FProfilerModule::CreateProfilerWindow( const ISessionManagerRef& InSessionManager, const TSharedRef<SDockTab>& ConstructUnderMajorTab )
+TSharedRef<SWidget> FProfilerModule::CreateProfilerWindow( const TSharedRef<ISessionManager>& InSessionManager, const TSharedRef<SDockTab>& ConstructUnderMajorTab )
 {
 	FProfilerManager::Initialize( InSessionManager );
 	TSharedRef<SProfilerWindow> ProfilerWindow = SNew( SProfilerWindow );
@@ -134,6 +138,7 @@ void FProfilerModule::StatsMemoryDumpCommand( const TCHAR* Filename )
 	}
 }
 
+
 /*-----------------------------------------------------------------------------
 	FRawStatsMemoryProfiler
 -----------------------------------------------------------------------------*/
diff --git a/Engine/Source/Developer/Profiler/Private/ProfilerPrivatePCH.h b/Engine/Source/Developer/Profiler/Private/ProfilerPrivatePCH.h
--- a/Engine/Source/Developer/Profiler/Private/ProfilerPrivatePCH.h
+++ b/Engine/Source/Developer/Profiler/Private/ProfilerPrivatePCH.h
@@ -2,49 +2,30 @@
 
 #pragma once
 
-#include "Profiler.h"
-
-
-/* Private dependencies
- *****************************************************************************/
-
-#include "ISessionInfo.h"
-#include "ISessionInstanceInfo.h"
-#include "ISessionManager.h"
-#include "Messaging.h"
-#include "ProfilerClient.h"
-#include "StatsData.h"
-#include "StatsFile.h"
-
-
-/* Private includes
- *****************************************************************************/
-
-// TODO: Remove dependencies between generic profiler code and profiler ui code
-
-// Common headers and managers.
-#include "ProfilerRawStatsForMemory.h"
-#include "ProfilerStream.h"
-
-#include "ProfilerSample.h"
-#include "ProfilerFPSAnalyzer.h"
-#include "ProfilerDataProvider.h"
-#include "ProfilerDataSource.h"
-#include "ProfilerSession.h"
-#include "ProfilerRawStatsForThreadView.h"
-#include "ProfilerCommands.h"
-#include "ProfilerManager.h"
-
-// Slate related headers.
-#include "SProfilerToolbar.h"
-#include "SFiltersAndPresets.h"
-#include "SMultiDumpBrowser.h"
-#include "SEventGraph.h"
-#include "SDataGraph.h"
-#include "SProfilerMiniView.h"
-#include "SProfilerThreadView.h"
-#include "SProfilerGraphPanel.h"
-#include "SProfilerFPSChartPanel.h"
-#include "SProfilerWindow.h"
-#include "SHistogram.h"
-#include "StatDragDropOp.h"
+#include "Runtime/Core/Public/Core.h"
+#include "Runtime/Core/Public/HAL/FileManagerGeneric.h"
+#include "Runtime/Core/Public/ProfilingDebugging/DiagnosticTable.h"
+#include "Runtime/Core/Public/Stats/StatsData.h"
+#include "Runtime/Core/Public/Stats/StatsFile.h"
+#include "Runtime/RenderCore/Public/RenderingThread.h"
+#include "Runtime/SessionServices/Public/ISessionInfo.h"
+#include "Runtime/SessionServices/Public/ISessionInstanceInfo.h"
+#include "Runtime/SessionServices/Public/ISessionManager.h"
+#include "Runtime/Slate/Public/SlateBasics.h"
+#include "Runtime/Slate/Public/Widgets/Docking/SDockTab.h"
+#include "Runtime/Slate/Public/Widgets/Input/SSearchBox.h"
+#include "Runtime/Slate/Public/Widgets/Notifications/SNotificationList.h"
+
+#include "Developer/DesktopPlatform/Public/DesktopPlatformModule.h"
+#include "Developer/ProfilerClient/Public/IProfilerClient.h"
+#include "Developer/ProfilerClient/Public/IProfilerClientModule.h"
+
+#include "Editor/EditorStyle/Public/EditorStyleSet.h" // @todo profiler: fix Editor style dependencies
+
+#if WITH_EDITOR
+	#include "Runtime/Analytics/Analytics/Public/Interfaces/IAnalyticsProvider.h"
+	#include "Runtime/Engine/Public/EngineAnalytics.h"
+#endif
+
+#include "IProfilerModule.h"
+#include "ProfilerCommon.h"
diff --git a/Engine/Source/Developer/Profiler/Private/ProfilerRawStatsForMemory.cpp b/Engine/Source/Developer/Profiler/Private/ProfilerRawStatsForMemory.cpp
--- a/Engine/Source/Developer/Profiler/Private/ProfilerRawStatsForMemory.cpp
+++ b/Engine/Source/Developer/Profiler/Private/ProfilerRawStatsForMemory.cpp
@@ -1,7 +1,8 @@
 // Copyright 1998-2016 Epic Games, Inc. All Rights Reserved.
 
 #include "ProfilerPrivatePCH.h"
-#include "DiagnosticTable.h"
+#include "ProfilerRawStatsForMemory.h"
+
 
 /*-----------------------------------------------------------------------------
 	Sort helpers
diff --git a/Engine/Source/Developer/Profiler/Private/ProfilerRawStatsForThreadView.cpp b/Engine/Source/Developer/Profiler/Private/ProfilerRawStatsForThreadView.cpp
--- a/Engine/Source/Developer/Profiler/Private/ProfilerRawStatsForThreadView.cpp
+++ b/Engine/Source/Developer/Profiler/Private/ProfilerRawStatsForThreadView.cpp
@@ -1,6 +1,9 @@
 // Copyright 1998-2015 Epic Games, Inc. All Rights Reserved.
 
 #include "ProfilerPrivatePCH.h"
+#include "ProfilerDataProvider.h"
+#include "ProfilerRawStatsForThreadView.h"
+
 
 // Only copied from ProfilerSession, still not working.
 
@@ -9,8 +12,8 @@
 -----------------------------------------------------------------------------*/
 
 FRawProfilerSession::FRawProfilerSession( const FString& InRawStatsFileFileath )
-: FProfilerSession( EProfilerSessionTypes::StatsFileRaw, nullptr, FGuid::NewGuid(), InRawStatsFileFileath.Replace( *FStatConstants::StatsFileRawExtension, TEXT( "" ) ) )
-, CurrentMiniViewFrame( 0 )
+	: FProfilerSession( EProfilerSessionTypes::StatsFileRaw, nullptr, FGuid::NewGuid(), InRawStatsFileFileath.Replace( *FStatConstants::StatsFileRawExtension, TEXT( "" ) ) )
+	, CurrentMiniViewFrame( 0 )
 {
 	OnTick = FTickerDelegate::CreateRaw( this, &FRawProfilerSession::HandleTicker );
 }
@@ -377,7 +380,7 @@ void FRawProfilerSession::ProcessStatPacketArray( const FStatPacketArray& StatPa
 	// Raw stats callstack for this stat packet array.
 	TMap<FName, FProfilerStackNode*> ThreadNodes;
 
-	const FProfilerStatMetaDataRef MetaData = GetMetaData();
+	const TSharedRef<FProfilerStatMetaData> MetaData = GetMetaData();
 
 	FProfilerSampleArray& MutableCollection = const_cast<FProfilerSampleArray&>(DataProvider->GetCollection());
 
diff --git a/Engine/Source/Developer/Profiler/Private/ProfilerRawStatsForThreadView.h b/Engine/Source/Developer/Profiler/Private/ProfilerRawStatsForThreadView.h
--- a/Engine/Source/Developer/Profiler/Private/ProfilerRawStatsForThreadView.h
+++ b/Engine/Source/Developer/Profiler/Private/ProfilerRawStatsForThreadView.h
@@ -2,14 +2,15 @@
 
 #pragma once
 
+#include "ProfilerSession.h"
+#include "ProfilerStream.h"
 
-/*-----------------------------------------------------------------------------
-	Profiler session for the raw stats files
------------------------------------------------------------------------------*/
 
 struct FAllocationInfo;
 
-class FRawProfilerSession : public FProfilerSession
+
+class FRawProfilerSession
+	: public FProfilerSession
 {
 	friend class FProfilerManager;
 	friend class FProfilerActionManager;
@@ -25,16 +26,17 @@ class FRawProfilerSession : public FProfilerSession
 	int32 CurrentMiniViewFrame;
 
 public:
+
 	/**
 	* Default constructor, creates a profiler session from a capture file.
 	*/
-	FRawProfilerSession( const FString& InRawStatsFileFileath );
+	FRawProfilerSession(const FString& InRawStatsFileFileath);
 
 	/** Destructor. */
 	~FRawProfilerSession();
 
 	/** Updates this profiler session. */
-	bool HandleTicker( float DeltaTime );
+	bool HandleTicker(float DeltaTime);
 
 	/** Starts a process of loading the raw stats file. */
 	void PrepareLoading();
@@ -48,5 +50,5 @@ public:
 	 *	Process all stats packets and convert them to data accessible by the profiler.
 	 *	Temporary version, will be optimized later.
 	 */
-	void ProcessStatPacketArray( const FStatPacketArray& PacketArray, FProfilerFrame& out_ProfilerFrame, int32 FrameIndex );
-};
\ No newline at end of file
+	void ProcessStatPacketArray(const FStatPacketArray& PacketArray, FProfilerFrame& out_ProfilerFrame, int32 FrameIndex);
+};
diff --git a/Engine/Source/Developer/Profiler/Private/ProfilerSample.cpp b/Engine/Source/Developer/Profiler/Private/ProfilerSample.cpp
--- a/Engine/Source/Developer/Profiler/Private/ProfilerSample.cpp
+++ b/Engine/Source/Developer/Profiler/Private/ProfilerSample.cpp
@@ -1,53 +1,61 @@
 // Copyright 1998-2016 Epic Games, Inc. All Rights Reserved.
 
 #include "ProfilerPrivatePCH.h"
+#include "ProfilerSample.h"
+
+
 #define LOCTEXT_NAMESPACE "ProfilerSample"
 
-/*-----------------------------------------------------------------------------
-	Enumerators
------------------------------------------------------------------------------*/
 
-FString EProfilerSampleTypes::ToName( Type ProfilerSampleType )
+FString EProfilerSampleTypes::ToName(Type ProfilerSampleType)
 {
-	switch( ProfilerSampleType )
+	switch(ProfilerSampleType)
 	{
-		case HierarchicalTime: return LOCTEXT("StatType_Name_Hierarchical","Hier").ToString();
-		case NumberInt: return LOCTEXT("StatType_Name_NumberInt","Int").ToString();
-		case NumberFloat: return LOCTEXT("StatType_Name_NumberFloat","Float").ToString();
-		case Memory: return LOCTEXT("StatType_Name_Memory","Mem").ToString();
-
-		default: return LOCTEXT("StatType_InvalidOrMax","Unknown profiler sample type").ToString();
+		case HierarchicalTime:
+			return LOCTEXT("StatType_Name_Hierarchical","Hier").ToString();
+		case NumberInt:
+			return LOCTEXT("StatType_Name_NumberInt","Int").ToString();
+		case NumberFloat:
+			return LOCTEXT("StatType_Name_NumberFloat","Float").ToString();
+		case Memory:
+			return LOCTEXT("StatType_Name_Memory","Mem").ToString();
+		default:
+			return LOCTEXT("StatType_InvalidOrMax","Unknown profiler sample type").ToString();
 	}
 }
 
-FString EProfilerSampleTypes::ToDescription( Type ProfilerSampleType )
+
+FString EProfilerSampleTypes::ToDescription(Type ProfilerSampleType)
 {
-	switch( ProfilerSampleType )
+	switch(ProfilerSampleType)
 	{
-		case HierarchicalTime: return LOCTEXT("StatType_Desc_Hierarchical","Hierarchical - Displayed as a time and call count").ToString();
-		case NumberInt: return LOCTEXT("StatType_Desc_NumberInt","Numerical - Displayed as a integer number").ToString();
-		case NumberFloat: return LOCTEXT("StatType_Desc_NumberFloat","Numerical - Displayed as a floating number").ToString();
-		case Memory: return LOCTEXT("StatType_Desc_Memory","Memory - Displayed as a human readable data counter").ToString();
-
-		default: return LOCTEXT("StatDesc_InvalidOrMax","Unknown profiler sample type").ToString();
+		case HierarchicalTime:
+			return LOCTEXT("StatType_Desc_Hierarchical","Hierarchical - Displayed as a time and call count").ToString();
+		case NumberInt:
+			return LOCTEXT("StatType_Desc_NumberInt","Numerical - Displayed as a integer number").ToString();
+		case NumberFloat:
+			return LOCTEXT("StatType_Desc_NumberFloat","Numerical - Displayed as a floating number").ToString();
+		case Memory:
+			return LOCTEXT("StatType_Desc_Memory","Memory - Displayed as a human readable data counter").ToString();
+		default:
+			return LOCTEXT("StatDesc_InvalidOrMax","Unknown profiler sample type").ToString();
 	}
 }
 
-/*-----------------------------------------------------------------------------
-	FProfilerSample
------------------------------------------------------------------------------*/
 
 const FProfilerSample FProfilerSample::Invalid;
 
-void FProfilerSample::FixupChildrenOrdering( const TMap<uint32,uint32>& ChildrenOrderingIndices )
+
+void FProfilerSample::FixupChildrenOrdering(const TMap<uint32,uint32>& ChildrenOrderingIndices)
 {
 	FIndicesArray ChildrenIndices;
-	for( auto It = ChildrenOrderingIndices.CreateConstIterator(); It; ++It )
+	for(auto It = ChildrenOrderingIndices.CreateConstIterator(); It; ++It)
 	{
-		ChildrenIndices.Add( _ChildrenIndices[It.Key()] );
+		ChildrenIndices.Add(_ChildrenIndices[It.Key()]);
 	}
 
 	_ChildrenIndices = ChildrenIndices;
 }
 
-#undef LOCTEXT_NAMESPACE
\ No newline at end of file
+
+#undef LOCTEXT_NAMESPACE
diff --git a/Engine/Source/Developer/Profiler/Private/ProfilerSession.cpp b/Engine/Source/Developer/Profiler/Private/ProfilerSession.cpp
--- a/Engine/Source/Developer/Profiler/Private/ProfilerSession.cpp
+++ b/Engine/Source/Developer/Profiler/Private/ProfilerSession.cpp
@@ -1,6 +1,10 @@
 // Copyright 1998-2016 Epic Games, Inc. All Rights Reserved.
 
 #include "ProfilerPrivatePCH.h"
+#include "ProfilerDataProvider.h"
+#include "ProfilerFPSAnalyzer.h"
+#include "ProfilerSession.h"
+
 
 /*-----------------------------------------------------------------------------
 	FProfilerStat, FProfilerGroup
@@ -20,7 +24,7 @@ FProfilerStat::FProfilerStat( const uint32 InStatID /*= 0 */ )
 	FProfilerSession
 -----------------------------------------------------------------------------*/
 
-FProfilerSession::FProfilerSession( EProfilerSessionTypes::Type InSessionType, const ISessionInstanceInfoPtr InSessionInstanceInfo, FGuid InSessionInstanceID, FString InDataFilepath )
+FProfilerSession::FProfilerSession( EProfilerSessionTypes InSessionType, const TSharedPtr<ISessionInstanceInfo>& InSessionInstanceInfo, FGuid InSessionInstanceID, FString InDataFilepath )
 : bRequestStatMetadataUpdate( false )
 , bLastPacket( false )
 , StatMetaDataSize( 0 )
@@ -42,7 +46,7 @@ FProfilerSession::FProfilerSession( EProfilerSessionTypes::Type InSessionType, c
 , FPSAnalyzer( MakeShareable( new FFPSAnalyzer( 5, 0, 90 ) ) )
 {}
 
-FProfilerSession::FProfilerSession( const ISessionInstanceInfoPtr InSessionInstanceInfo )
+FProfilerSession::FProfilerSession( const TSharedPtr<ISessionInstanceInfo>& InSessionInstanceInfo )
 : bRequestStatMetadataUpdate( false )
 , bLastPacket( false )
 , StatMetaDataSize( 0 )
@@ -99,22 +103,22 @@ FProfilerSession::~FProfilerSession()
 	FTicker::GetCoreTicker().RemoveTicker( OnTickHandle );	
 }
 
-const FEventGraphDataRef FProfilerSession::GetEventGraphDataTotal() const
+const TSharedRef<FEventGraphData, ESPMode::ThreadSafe> FProfilerSession::GetEventGraphDataTotal() const
 {
 	return EventGraphDataTotal.ToSharedRef();
 }
 
-const FEventGraphDataRef FProfilerSession::GetEventGraphDataMaximum() const
+const TSharedRef<FEventGraphData, ESPMode::ThreadSafe> FProfilerSession::GetEventGraphDataMaximum() const
 {
 	return EventGraphDataMaximum.ToSharedRef();
 }
 
-const FEventGraphDataRef FProfilerSession::GetEventGraphDataAverage() const
+const TSharedRef<FEventGraphData, ESPMode::ThreadSafe> FProfilerSession::GetEventGraphDataAverage() const
 {
 	return EventGraphDataAverage.ToSharedRef();
 }
 
-FGraphDataSourceRefConst FProfilerSession::CreateGraphDataSource( const uint32 InStatID )
+TSharedRef<const FGraphDataSource> FProfilerSession::CreateGraphDataSource( const uint32 InStatID )
 {
 	FGraphDataSource* GraphDataSource = new FGraphDataSource( AsShared(), InStatID );
 	return MakeShareable( GraphDataSource );
@@ -273,12 +277,11 @@ FEventGraphContainer FProfilerSession::CreateEventGraphData( const uint32 FrameS
 	check( EventGraphData );
 	EventGraphData->Finalize( FrameStartIndex, FrameEndIndex + 1 );
 
-	FEventGraphDataRef Total = MakeShareable( EventGraphData );
-
-	FEventGraphDataRef Average = Total->DuplicateAsRef();
+	TSharedRef<FEventGraphData, ESPMode::ThreadSafe> Total = MakeShareable( EventGraphData );
+	TSharedRef<FEventGraphData, ESPMode::ThreadSafe> Average = Total->DuplicateAsRef();
 	Average->SetAsAverage();
 
-	FEventGraphDataRef Maximum = Total->DuplicateAsRef();
+	TSharedRef<FEventGraphData, ESPMode::ThreadSafe> Maximum = Total->DuplicateAsRef();
 	Maximum->SetAsMaximim();
 	
 	FEventGraphContainer EventGraphContainer( FrameStartIndex, FrameEndIndex + 1, Average, Maximum, Total );
@@ -547,7 +550,7 @@ void FProfilerSession::PopulateHierarchy_Recurrent
 	TMap<uint32, FInclusiveTime>& StatIDToInclusiveTime
 )
 {
-	const FProfilerStatMetaDataRef MetaData = GetMetaData();
+	const TSharedRef<FProfilerStatMetaData> MetaData = GetMetaData();
 
 #if	UE_BUILD_DEBUG
 	const FName& ParentStatName = MetaData->GetStatByID( ParentGraph.StatId ).Name();
@@ -956,7 +959,7 @@ void FProfilerAggregatedStat::Advance()
 	_NumCallsOneFrame = 0;
 }
 
-void FProfilerAggregatedStat::Aggregate(const FProfilerSample& Sample, const FProfilerStatMetaDataRef& Metadata)
+void FProfilerAggregatedStat::Aggregate(const FProfilerSample& Sample, const TSharedRef<FProfilerStatMetaData>& Metadata)
 {
 	double TypedValue = 0.0;
 
diff --git a/Engine/Source/Developer/Profiler/Private/ProfilerSession.h b/Engine/Source/Developer/Profiler/Private/ProfilerSession.h
--- a/Engine/Source/Developer/Profiler/Private/ProfilerSession.h
+++ b/Engine/Source/Developer/Profiler/Private/ProfilerSession.h
@@ -2,112 +2,124 @@
 
 #pragma once
 
-/*-----------------------------------------------------------------------------
-	Enumerators
------------------------------------------------------------------------------*/
+#include "IProfilerServiceManager.h"
+#include "ProfilerDataSource.h"
+#include "ProfilerSample.h"
+#include "ProfilerStream.h"
 
-/** Enumerates profiler session type. */
-namespace EProfilerSessionTypes
-{
-	enum Type
-	{
-		/** Based on the live connection. */
-		Live,
-		LiveRaw,
 
-		/** Based on the regular stats file. */
-		StatsFile,
+class FEventGraphData;
+class FFPSAnalyzer;
+class FGraphDataSource;
+class FProfilerStatMetaData;
+class IDataProvider;
+class ISessionInstanceInfo;
 
-		/** Based on the raw stats file. */
-		StatsFileRaw,
 
-		Combined,
-		Summary,
+/**
+ * Enumerates profiler session type.
+ */
+enum class EProfilerSessionTypes
+{
+	/** Based on the live connection. */
+	Live,
+	LiveRaw,
 
-		/** Invalid enum type, may be used as a number of enumerations. */
-		InvalidOrMax,
+	/** Based on the regular stats file. */
+	StatsFile,
 
+	/** Based on the raw stats file. */
+	StatsFileRaw,
 
+	Combined,
+	Summary,
 
-	};
+	/** Invalid enum type, may be used as a number of enumerations. */
+	InvalidOrMax,
+};
 
+
+namespace ProfilerSessionTypes
+{
 	/**
 	 * Returns the string representation of the specified EProfilerSessionType value.
 	 *
-	 * @param ProfilerSessionType - The value to get the string for.
-	 *
+	 * @param ProfilerSessionType The value to get the string for.
 	 * @return A string value.
 	 */
-	inline FString ToString( EProfilerSessionTypes::Type ProfilerSessionType )
+	inline FString ToString(EProfilerSessionTypes ProfilerSessionType)
 	{
 		// @TODO: Localize
-		switch( ProfilerSessionType )
+		switch(ProfilerSessionType)
 		{
-			case Live: return FString("Live");
-			case StatsFile: return FString("Offline");
-
-			default: return FString("InvalidOrMax");
+		case EProfilerSessionTypes::Live:
+			return FString("Live");
+		case EProfilerSessionTypes::StatsFile:
+			return FString("Offline");
+		default:
+			return FString("InvalidOrMax");
 		}
 	}
 }
 
-/** Enumerates loading a capture file progress states. */
-namespace ELoadingProgressStates
+
+/**
+ * Enumerates loading a capture file progress states.
+ */
+enum class ELoadingProgressStates
 {
-	enum Type
-	{
-		Started,
-		InProgress,
-		Loaded,
-		Failed,
-		Cancelled, 
-
-		/** Invalid enum type, may be used as a number of enumerations. */
-		InvalidOrMax,
-	};
-}
+	Started,
+	InProgress,
+	Loaded,
+	Failed,
+	Cancelled, 
+
+	/** Invalid enum type, may be used as a number of enumerations. */
+	InvalidOrMax,
+};
+
 
-namespace EProfilerNotificationTypes
+enum class EProfilerNotificationTypes
 {
-	enum Type
-	{
-		LoadingOfflineCapture,
-		SendingServiceSideCapture,
-	};
-}
+	LoadingOfflineCapture,
+	SendingServiceSideCapture,
+};
 
-/*-----------------------------------------------------------------------------
-	Declarations
------------------------------------------------------------------------------*/
 
-/**	Contains information about stat like name, ID and type, and information about the owning stat group. */
-class FProfilerStat : public FNoncopyable
+class FProfilerGroup;
+
+
+/**
+ * Contains information about stat like name, ID and type, and information about the owning stat group.
+ */
+class FProfilerStat
+	: public FNoncopyable
 {
 	friend class FProfilerStatMetaData;
 
 public:
+
 	/**
 	 * Default constructor.
 	 *
-	 * @param InID		- the unique ID for this stat
-	 *
+	 * @param InID The unique ID for this stat
 	 */
-	FProfilerStat( const uint32 InStatID = 0 );
+	FProfilerStat(const uint32 InStatID = 0);
 
 protected:
+
 	/**
 	 * Sets the new description for this stat.
 	 *
-	 * @param InName	- the name of this stat
-	 * @param InGroupID - the pointer to the group that this stat belongs to
-	 * @param InType	- the type of this stat
-	 *
+	 * @param InName The name of this stat
+	 * @param InGroupID The pointer to the group that this stat belongs to
+	 * @param InType The type of this stat
 	 */
-	void Initialize( const FString& InName, FProfilerGroup* InOwningGroupPtr, const EStatType InType )	
+	void Initialize(const FString& InName, FProfilerGroup* InOwningGroupPtr, const EStatType InType)	
 	{
 		// Skip leading spaces in the name of the stat.
 		int32 Index = 0;
-		while( FChar::IsWhitespace( InName[Index] ) )
+		while(FChar::IsWhitespace(InName[Index]))
 		{
 			Index++;
 		}
@@ -121,11 +133,11 @@ protected:
 	/**
 	 * @return a profiler sample type from the specified generic stat type
 	 */
-	static const EProfilerSampleTypes::Type ConvertStatTypeToProfilerSampleType( const EStatType StatType )
+	static const EProfilerSampleTypes::Type ConvertStatTypeToProfilerSampleType(const EStatType StatType)
 	{
 		EProfilerSampleTypes::Type ProfilerUnit = EProfilerSampleTypes::InvalidOrMax;
 
-		switch( StatType )
+		switch(StatType)
 		{	
 		case STATTYPE_MemoryCounter:
 			{
@@ -173,7 +185,7 @@ public:
 	 */
 	const FProfilerGroup& OwningGroup() const
 	{
-		check( _OwningGroupPtr != NULL );
+		check(_OwningGroupPtr != NULL);
 		return *_OwningGroupPtr;
 	}
 
@@ -225,8 +237,10 @@ protected:
 	/**	Default profiler stat, used for uninitialized stats. */
 	static FProfilerStat Default;
 };
+
 template <> struct TIsPODType<FProfilerStat> { enum { Value = true }; };
 
+
 /**
  *	Contains information about stat group and stats associated with the specified stat group.
  */
@@ -235,33 +249,32 @@ class FProfilerGroup : public FNoncopyable
 	friend class FProfilerStatMetaData;
 
 public:
-	/**
-	 * Default constructor.
-	 */
+
+	/** Default constructor. */
 	FProfilerGroup()
-		: _Name( TEXT("(Group-Default)") )
-		, _ID( 0 )
+		: _Name(TEXT("(Group-Default)"))
+		, _ID(0)
 	{}
 
 protected:
 	/**
 	 * Initialization constructor.
 	 *
-	 * @param InID	- the unique ID for this stat group
+	 * @param InID The unique ID for this stat group
 	 *
 	 */
-	FProfilerGroup( const uint32 InID )
-		: _Name( *FString::Printf( TEXT("(Group-%04u)"), InID ) )
-		, _ID( InID )
+	FProfilerGroup(const uint32 InID)
+		: _Name(*FString::Printf(TEXT("(Group-%04u)"), InID))
+		, _ID(InID)
 	{}
 
 	/**
 	 * Sets the new name for this stat group.
 	 *
-	 * @param InName - the name of this stat group 
+	 * @param InName The name of this stat group 
 	 *
 	 */
-	void Initialize( const FString& InName )
+	void Initialize(const FString& InName)
 	{
 		_Name = *InName;
 	}
@@ -269,12 +282,12 @@ protected:
 	/**
 	 * Adds a stat to this group.
 	 *
-	 * @param ProfilerStat - a pointer to the existing stat object
+	 * @param ProfilerStat Pointer to the existing stat object
 	 * 
 	 */
-	void AddStat( FProfilerStat* ProfilerStat )
+	void AddStat(FProfilerStat* ProfilerStat)
 	{
-		OwnedStats.Add( ProfilerStat );
+		OwnedStats.Add(ProfilerStat);
 	}
 
 public:
@@ -344,8 +357,8 @@ class FProfilerStatMetaData : public FNoncopyable, public TSharedFromThis<FProfi
 protected:
 	/** Constructor. */
 	FProfilerStatMetaData()
-		: SecondsPerCycle( FPlatformTime::GetSecondsPerCycle() )
-		, GameThreadID( 0 )
+		: SecondsPerCycle(FPlatformTime::GetSecondsPerCycle())
+		, GameThreadID(0)
 	{}
 
 public:
@@ -353,13 +366,13 @@ public:
 	~FProfilerStatMetaData()
 	{
 		// Delete all allocated descriptions.
-		for( auto It = StatDescriptions.CreateConstIterator(); It; ++It )
+		for(auto It = StatDescriptions.CreateConstIterator(); It; ++It)
 		{
 			delete It.Value();
 		}
 		StatDescriptions.Empty();
 
-		for( auto It = GroupDescriptions.CreateConstIterator(); It; ++It )
+		for(auto It = GroupDescriptions.CreateConstIterator(); It; ++It)
 		{
 			delete It.Value();
 		}
@@ -375,13 +388,13 @@ public:
 
 		MemorySize += sizeof(*this);
 		MemorySize += StatDescriptions.GetAllocatedSize();
-		for( auto It = StatDescriptions.CreateConstIterator(); It; ++It )
+		for(auto It = StatDescriptions.CreateConstIterator(); It; ++It)
 		{
 			MemorySize += It.Value()->GetMemoryUsage();
 		}
 
 		MemorySize += GroupDescriptions.GetAllocatedSize();
-		for( auto It = GroupDescriptions.CreateConstIterator(); It; ++It )
+		for(auto It = GroupDescriptions.CreateConstIterator(); It; ++It)
 		{
 			MemorySize += It.Value()->GetMemoryUsage();
 		}
@@ -395,17 +408,16 @@ protected:
 	/**
 	 * Updates this instance of stat metadata.
 	 *
-	 * @param StatMetaData - the stat metadata received from the connected profiler client
-	 *
+	 * @param StatMetaData The stat metadata received from the connected profiler client
 	 */
-	void Update( const FStatMetaData& ClientStatMetaData );
+	void Update(const FStatMetaData& ClientStatMetaData);
 
 public:
 	/**
 	 *	Updates this stats metadata based on the stats thread state.
 	 *	This is a temporary solution to make it working with current implementation of FProfilerSample.
 	 */
-	void UpdateFromStatsState( const FStatsThreadState& StatsThreadStats );
+	void UpdateFromStatsState(const FStatsThreadState& StatsThreadStats);
 
 private:
 
@@ -413,82 +425,80 @@ private:
 	 * Initialized the specified stat group.
 	 * If specified stat group doesn't exist, adds it to the list of stat group descriptions.
 	 *
-	 * @param GroupID	- the ID of the stat group that needs to be initialized
-	 * @param GroupName	- the name of the stat group that needs to be initialized
-	 *
+	 * @param GroupID The ID of the stat group that needs to be initialized
+	 * @param GroupName The name of the stat group that needs to be initialized
 	 */
-	void InitializeGroup( const uint32 GroupID, const FString& GroupName ) 
+	void InitializeGroup(const uint32 GroupID, const FString& GroupName) 
 	{
-		FProfilerGroup* GroupPtr = GroupDescriptions.FindRef( GroupID );
-		if( !GroupPtr )
+		FProfilerGroup* GroupPtr = GroupDescriptions.FindRef(GroupID);
+		if(!GroupPtr)
 		{
-			GroupPtr = GroupDescriptions.Add( GroupID, new FProfilerGroup( GroupID ) );
+			GroupPtr = GroupDescriptions.Add(GroupID, new FProfilerGroup(GroupID));
 		}
-		GroupPtr->Initialize( GroupName );
+		GroupPtr->Initialize(GroupName);
 	}
 
 	/**
 	 * Initialized the specified stat.
 	 *
-	 * @param StatID	- the id of the stat that needs to be initialized
-	 * @param GroupID	- the id of the stat group that will be assigned to the specified stat
-	 * @param StatName  - the name of the stat that needs to be initialized
-	 * @param InType	- the type of the stat that needs to be initialized
-	 * @param StatName	- the unique stat name
-	 *
+	 * @param StatID The id of the stat that needs to be initialized
+	 * @param GroupID The id of the stat group that will be assigned to the specified stat
+	 * @param StatName The name of the stat that needs to be initialized
+	 * @param InType The type of the stat that needs to be initialized
+	 * @param StatName The unique stat name
 	 */
-	void InitializeStat( const uint32 StatID, const uint32 GroupID, const FString& StatName, const EStatType InType, FName StatFName = NAME_None )	
+	void InitializeStat(const uint32 StatID, const uint32 GroupID, const FString& StatName, const EStatType InType, FName StatFName = NAME_None)	
 	{
-		FProfilerStat* StatPtr = StatDescriptions.FindRef( StatID );
-		if( !StatPtr )
+		FProfilerStat* StatPtr = StatDescriptions.FindRef(StatID);
+		if(!StatPtr)
 		{
-			StatPtr = StatDescriptions.Add( StatID, new FProfilerStat( StatID ) );
+			StatPtr = StatDescriptions.Add(StatID, new FProfilerStat(StatID));
 
-			if( StatFName != NAME_None )
+			if(StatFName != NAME_None)
 			{
-				StatFNameDescriptions.Add( StatFName, StatPtr );
+				StatFNameDescriptions.Add(StatFName, StatPtr);
 			}
 
-			FProfilerGroup* GroupPtr = GroupDescriptions.FindRef( GroupID );
+			FProfilerGroup* GroupPtr = GroupDescriptions.FindRef(GroupID);
 
-			StatPtr->Initialize( StatName, GroupPtr, InType );
+			StatPtr->Initialize(StatName, GroupPtr, InType);
 
 			static const FName NAME_Threads(TEXT("Threads"));
-			if( StatFName == NAME_None && GroupPtr->Name() == NAME_Threads )
+			if(StatFName == NAME_None && GroupPtr->Name() == NAME_Threads)
 			{
 				// Check if this stat is a thread stat.
-				const uint32 ThreadID = FStatsUtils::ParseThreadID( StatPtr->Name().ToString() );
-				if( ThreadID != 0 )
+				const uint32 ThreadID = FStatsUtils::ParseThreadID(StatPtr->Name().ToString());
+				if(ThreadID != 0)
 				{
-					const FString* ThreadDesc = GetThreadDescriptions().Find( ThreadID );
-					if( ThreadDesc )
+					const FString* ThreadDesc = GetThreadDescriptions().Find(ThreadID);
+					if(ThreadDesc)
 					{
 						// Replace the stat name with a thread name.
-						const FString UniqueThreadName = FString::Printf( TEXT( "%s [0x%x]" ), **ThreadDesc, ThreadID );
+						const FString UniqueThreadName = FString::Printf(TEXT("%s [0x%x]"), **ThreadDesc, ThreadID);
 						StatPtr->_Name = *UniqueThreadName;
-						ThreadIDtoStatID.Add( ThreadID, StatID );
+						ThreadIDtoStatID.Add(ThreadID, StatID);
 
 						// Game thread is always NAME_GameThread
-						if( **ThreadDesc == FName( NAME_GameThread ) )
+						if(**ThreadDesc == FName(NAME_GameThread))
 						{
 							GameThreadID = ThreadID;
 						}
 						// Rendering thread may be "Rendering thread" or NAME_RenderThread with an index
-						else if( ThreadDesc->Contains( FName( NAME_RenderThread ).GetPlainNameString() ) )
+						else if(ThreadDesc->Contains(FName(NAME_RenderThread).GetPlainNameString()))
 						{
-							RenderThreadIDs.AddUnique( ThreadID );
+							RenderThreadIDs.AddUnique(ThreadID);
 						}
-						else if( ThreadDesc->Contains( TEXT( "RenderingThread" ) ) )
+						else if(ThreadDesc->Contains(TEXT("RenderingThread")))
 						{
-							RenderThreadIDs.AddUnique( ThreadID );
+							RenderThreadIDs.AddUnique(ThreadID);
 						}
 					}			
 				}
 			}
 
-			if( GroupPtr != FProfilerGroup::GetDefaultPtr() )
+			if(GroupPtr != FProfilerGroup::GetDefaultPtr())
 			{
-				GroupPtr->AddStat( StatPtr );
+				GroupPtr->AddStat(StatPtr);
 			}
 		}
 	}
@@ -497,17 +507,17 @@ public:
 	/**
 	 * @return a reference to the stat description specified by the stat ID
 	 */
-	const FProfilerStat& GetStatByID( const uint32 StatID ) const
+	const FProfilerStat& GetStatByID(const uint32 StatID) const
 	{
-		return *StatDescriptions.FindChecked( StatID );
+		return *StatDescriptions.FindChecked(StatID);
 	}
 
 	/**
 	* @return a reference to the stat description specified by the unique stat name.
 	*/
-	const FProfilerStat& GetStatByFName( FName StatName ) const
+	const FProfilerStat& GetStatByFName(FName StatName) const
 	{
-		return *StatFNameDescriptions.FindChecked( StatName );
+		return *StatFNameDescriptions.FindChecked(StatName);
 	}
 
 	/**
@@ -521,9 +531,9 @@ public:
 	/**
 	 * @return a reference to the group description specified by the group ID
 	 */
-	const FProfilerGroup& GetGroup( const uint32 GroupID ) const
+	const FProfilerGroup& GetGroup(const uint32 GroupID) const
 	{
-		return *GroupDescriptions.FindChecked( GroupID );
+		return *GroupDescriptions.FindChecked(GroupID);
 	}
 
 	/**
@@ -553,7 +563,7 @@ public:
 	/**
 	 * @return the specified number of cycles converted to milliseconds.
 	 */
-	const double ConvertCyclesToMS( const uint32 Cycles ) const
+	const double ConvertCyclesToMS(const uint32 Cycles) const
 	{
 		return SecondsPerCycle * 1000 * Cycles;
 	}
@@ -561,17 +571,17 @@ public:
 	/**
 	 * @return the profiler sample type for the specified stat ID.
 	 */
-	const EProfilerSampleTypes::Type GetSampleTypeForStatID( const uint32 StatID ) const
+	const EProfilerSampleTypes::Type GetSampleTypeForStatID(const uint32 StatID) const
 	{
-		return GetStatByID( StatID ).Type();
+		return GetStatByID(StatID).Type();
 	}
 
 	/**
 	 * @return true, if the stat for the specified stat ID has been initialized.
 	 */
-	const bool IsStatInitialized( const uint32 StatID )
+	const bool IsStatInitialized(const uint32 StatID)
 	{
-		return StatDescriptions.Contains( StatID );
+		return StatDescriptions.Contains(StatID);
 	}
 
 	const uint32 GetGameThreadID() const
@@ -591,7 +601,7 @@ public:
 
 	const uint32 GetGameThreadStatID() const
 	{
-		return ThreadIDtoStatID.FindChecked( GameThreadID );
+		return ThreadIDtoStatID.FindChecked(GameThreadID);
 	}
 
 protected:
@@ -623,6 +633,7 @@ protected:
 	TArray<uint32> RenderThreadIDs;
 };
 
+
 /** Holds the aggregated information for the specific stat across all frames that have been captured. */
 class FProfilerAggregatedStat
 {
@@ -645,7 +656,7 @@ public:
 	};
 
 	/** Default constructor. */
-	FProfilerAggregatedStat( const FName& InStatName, const FName& InGroupName, EProfilerSampleTypes::Type InStatType );
+	FProfilerAggregatedStat(const FName& InStatName, const FName& InGroupName, EProfilerSampleTypes::Type InStatType);
 
 public:
 	/** @return The average value of all combined instances. */
@@ -710,7 +721,7 @@ public:
 	/**
 	 * @return a string representation of the specified value type.
 	 */
-	FString GetFormattedValue( const Type ValueType ) const;
+	FString GetFormattedValue(const Type ValueType) const;
 
 protected:
 	/** Called once a frame to update aggregates. */
@@ -722,7 +733,7 @@ protected:
 	 * @param Sample - a reference to the profiler sample that will be aggregated
 	 *
 	 */
-	void Aggregate( const FProfilerSample& Sample, const FProfilerStatMetaDataRef& Metadata );
+	void Aggregate(const FProfilerSample& Sample, const TSharedRef<FProfilerStatMetaData>& Metadata);
 
 protected:
 	FName _StatName;
@@ -764,48 +775,57 @@ protected:
 	const EProfilerSampleTypes::Type StatType;
 };
 
+
 //template <> struct TIsPODType<FProfilerAggregatedStat> { enum { Value = true }; };
 
 // @TODO: Standardize naming used in the profiler session/ session / session instance.
 
-/** Class that holds the profiler session information. */
-class FProfilerSession : public TSharedFromThis<FProfilerSession>
+
+/**
+ * Class that holds the profiler session information.
+ */
+class FProfilerSession
+	: public TSharedFromThis<FProfilerSession>
 {
 	friend class FProfilerManager;
 	friend class FProfilerActionManager;
 
 protected:
+
 	/** Default constructor. */
-	FProfilerSession
-	( 
-		EProfilerSessionTypes::Type InSessionType, 
-		const ISessionInstanceInfoPtr InSessionInstanceInfo,
+	FProfilerSession(
+		EProfilerSessionTypes InSessionType, 
+		const TSharedPtr<ISessionInstanceInfo>& InSessionInstanceInfo,
 		FGuid InSessionInstanceID,
-		FString InDataFilepath
-	);
+		FString InDataFilepath);
 
 public:
+
 	/**
 	 * Initialization constructor, creates a profiler session from a capture file.
 	 */
-	FProfilerSession( const FString& InDataFilepath );
+	FProfilerSession(const FString& InDataFilepath);
 
 	/**
 	 * Initialization constructor, creates a live profiler session.
 	 */
-	FProfilerSession( const ISessionInstanceInfoPtr InSessionInstanceInfo );
+	FProfilerSession(const TSharedPtr<ISessionInstanceInfo>& InSessionInstanceInfo);
 
 	/** Destructor. */
 	~FProfilerSession();
 
+public:
+
 	/** Updates this profiler session. */
-	bool HandleTicker( float DeltaTime );
+	bool HandleTicker(float DeltaTime);
 
 public:
+
 	typedef TMap<uint32, float> FThreadTimesMap;
-	DECLARE_DELEGATE_ThreeParams( FAddThreadTimeDelegate, int32 /*FrameIndex*/, const FThreadTimesMap& /*ThreadMS*/, const FProfilerStatMetaDataRef& /*StatMetaData*/ );
+	typedef TSharedRef<FProfilerStatMetaData> FProfilerStatMetaDataRef;
+	DECLARE_DELEGATE_ThreeParams(FAddThreadTimeDelegate, int32 /*FrameIndex*/, const FThreadTimesMap& /*ThreadMS*/, const FProfilerStatMetaDataRef& /*StatMetaData*/);
 
-	FProfilerSession& SetOnAddThreadTime( const FAddThreadTimeDelegate& InOnAddThreadTime )
+	FProfilerSession& SetOnAddThreadTime(const FAddThreadTimeDelegate& InOnAddThreadTime)
 	{
 		OnAddThreadTime = InOnAddThreadTime;
 		return *this;
@@ -821,14 +841,14 @@ public:
 	 * @param SessionInstanceID - the session instance ID of this session, used to determine of which session this method should be called
 	 *
 	 */
-	DECLARE_DELEGATE_OneParam( FCaptureFileProcessedDelegate, const FGuid /*SessionInstanceID*/ );
+	DECLARE_DELEGATE_OneParam(FCaptureFileProcessedDelegate, const FGuid /*SessionInstanceID*/);
 
 	/**
 	 * Sets the handler that will be invoked when this profiler session has fully processed a capture file
 	 * 
 	 * @return This instance (for method chaining).
 	 */
-	FProfilerSession& SetOnCaptureFileProcessed( const FCaptureFileProcessedDelegate& InOnCaptureFileProcessed )
+	FProfilerSession& SetOnCaptureFileProcessed(const FCaptureFileProcessedDelegate& InOnCaptureFileProcessed)
 	{
 		OnCaptureFileProcessed = InOnCaptureFileProcessed;
 		return *this;
@@ -846,17 +866,17 @@ public:
 	const FString GetName() const
 	{
 		FString SessionName;
-		if( SessionType == EProfilerSessionTypes::Live )
+		if(SessionType == EProfilerSessionTypes::Live)
 		{
-			SessionName = FString::Printf( TEXT("%s"), *SessionInstanceInfo->GetInstanceName() );
+			SessionName = FString::Printf(TEXT("%s"), *SessionInstanceInfo->GetInstanceName());
 		}
-		else if( SessionType == EProfilerSessionTypes::StatsFile )
+		else if(SessionType == EProfilerSessionTypes::StatsFile)
 		{
-			SessionName = FString::Printf( TEXT("%s"), *DataFilepath );
+			SessionName = FString::Printf(TEXT("%s"), *DataFilepath);
 		}
-		else if( SessionType == EProfilerSessionTypes::StatsFileRaw )
+		else if(SessionType == EProfilerSessionTypes::StatsFileRaw)
 		{
-			SessionName = FString::Printf( TEXT( "%s" ), *DataFilepath );
+			SessionName = FString::Printf(TEXT("%s"), *DataFilepath);
 		}
 		
 		return SessionName;
@@ -867,13 +887,13 @@ public:
 	 */
 	const FString GetShortName() const
 	{
-		return FProfilerHelper::ShortenName( GetName(), 12 );
+		return FProfilerHelper::ShortenName(GetName(), 12);
 	}
 
 	/**
 	 * @return session type for this profiler session
 	 */
-	const EProfilerSessionTypes::Type GetSessionType() const
+	const EProfilerSessionTypes GetSessionType() const
 	{
 		return SessionType;
 	}
@@ -887,7 +907,7 @@ public:
 	}
 
 	/**
-	 * @return the time when this profiler session was created ( time of the connection to the client, time when a profiler capture was created ).
+	 * @return the time when this profiler session was created (time of the connection to the client, time when a profiler capture was created).
 	 */
 	const FDateTime& GetCreationTime() const
 	{
@@ -897,7 +917,7 @@ public:
 	/**
 	 * @return a shared reference to the the data provider which holds all the collected profiler samples.
 	 */
-	const IDataProviderRef& GetDataProvider() const
+	const TSharedRef<IDataProvider>& GetDataProvider() const
 	{
 		return DataProvider;
 	}
@@ -905,7 +925,7 @@ public:
 	/**
 	 * @return a shared reference to the stat metadata which holds all collected stats descriptions.
 	 */
-	const FProfilerStatMetaDataRef& GetMetaData() const
+	const TSharedRef<FProfilerStatMetaData>& GetMetaData() const
 	{
 		return StatMetaData;
 	}
@@ -913,19 +933,19 @@ public:
 	/**
 	 * @return a const pointer to the aggregated stat for the specified stat ID or null if not found.
 	 */
-	FORCEINLINE_DEBUGGABLE const FProfilerAggregatedStat* GetAggregatedStat( const uint32 StatID ) const
+	FORCEINLINE_DEBUGGABLE const FProfilerAggregatedStat* GetAggregatedStat(const uint32 StatID) const
 	{
-		return AggregatedStats.Find( StatID );
+		return AggregatedStats.Find(StatID);
 	}
 
-	FORCEINLINE_DEBUGGABLE const TMap<uint32, FInclusiveTime>& GetInclusiveAggregateStackStats( const uint32 FrameIndex ) const
+	FORCEINLINE_DEBUGGABLE const TMap<uint32, FInclusiveTime>& GetInclusiveAggregateStackStats(const uint32 FrameIndex) const
 	{
 		return InclusiveAggregateStackStats[FrameIndex];
 	}
 
-	const FEventGraphDataRef GetEventGraphDataTotal() const;
-	const FEventGraphDataRef GetEventGraphDataMaximum() const;
-	const FEventGraphDataRef GetEventGraphDataAverage() const;
+	const TSharedRef<FEventGraphData, ESPMode::ThreadSafe> GetEventGraphDataTotal() const;
+	const TSharedRef<FEventGraphData, ESPMode::ThreadSafe> GetEventGraphDataMaximum() const;
+	const TSharedRef<FEventGraphData, ESPMode::ThreadSafe> GetEventGraphDataAverage() const;
 
 	/**
 	 * @return number of bytes allocated by this profiler session.
@@ -935,28 +955,27 @@ public:
 	/**
 	 * @return a new instance of the graph data source for the specified stat ID.
 	 */
-	FGraphDataSourceRefConst CreateGraphDataSource( const uint32 InStatID );
+	TSharedRef<const FGraphDataSource> CreateGraphDataSource(const uint32 InStatID);
 
 	/**
 	 * @return a new instance of the event graph container for the specified frame range.
 	 */
-	FEventGraphContainer CreateEventGraphData( const uint32 FrameStartIndex, const uint32 FrameEndIndex );
+	FEventGraphContainer CreateEventGraphData(const uint32 FrameStartIndex, const uint32 FrameEndIndex);
 
 	/** Combines event graphs for the specified frames range. */
-	FEventGraphData* CombineEventGraphs( const uint32 FrameStartIndex, const uint32 FrameEndIndex );
+	FEventGraphData* CombineEventGraphs(const uint32 FrameStartIndex, const uint32 FrameEndIndex);
 
 	/** Combines event graphs for the specified frames range, as the task graph task. */
-	void CombineEventGraphsTask( const uint32 FrameStartIndex, const uint32 FrameEndIndex );
+	void CombineEventGraphsTask(const uint32 FrameStartIndex, const uint32 FrameEndIndex);
 
 protected:
 	/**
 	 * Recursively populates hierarchy based on the specified profiler cycle graph.
 	 *
-	 * @param ParentGraph - <describe ParentGraph>
-	 * @param ParentStartTimeMS - <describe ParentStartTimeMS>
-	 * @param ParentDurationMS - <describe ParentDurationMS>
-	 * @param ParentSampleIndex - <describe ParentSampleIndex>
-	 *
+	 * @param ParentGraph
+	 * @param ParentStartTimeMS
+	 * @param ParentDurationMS
+	 * @param ParentSampleIndex
 	 * @return <describe return value>
 	 */
 	void PopulateHierarchy_Recurrent
@@ -969,10 +988,10 @@ protected:
 	);
 
 	/** Called when this profiler session receives a new profiler data. */
-	void UpdateProfilerData( const FProfilerDataFrame& Content );
+	void UpdateProfilerData(const FProfilerDataFrame& Content);
 
 	/** Called when this profiler session receives information that the meta data has been updated. @see FProfilerMetaDataUpdateDelegate and IProfilerClient */
-	void UpdateMetadata( const FStatMetaData& InClientStatMetaData );
+	void UpdateMetadata(const FStatMetaData& InClientStatMetaData);
 
 	/**
 	 * Updates aggregated stats.
@@ -980,7 +999,7 @@ protected:
 	 * @param FrameIndex - the index of the profiler samples that will be used to update the aggregated stats.
 	 *
 	 */
-	void UpdateAggregatedStats( const uint32 FrameIndex );
+	void UpdateAggregatedStats(const uint32 FrameIndex);
 
 	/**
 	 * Updates aggregated event graphs.
@@ -988,20 +1007,20 @@ protected:
 	 * @param FrameIndex - the index of the profiler samples that will be used to update the aggregated stats.
 	 *
 	 */
-	void UpdateAggregatedEventGraphData( const uint32 FrameIndex );
+	void UpdateAggregatedEventGraphData(const uint32 FrameIndex);
 
 	/** Completion sync. */
 	void CompletionSyncAggregatedEventGraphData();
 
-	void EventGraphCombine( const FEventGraphData* Current, const uint32 InNumFrames );
+	void EventGraphCombine(const FEventGraphData* Current, const uint32 InNumFrames);
 
-	void UpdateAllEventGraphs( const uint32 InNumFrames );
+	void UpdateAllEventGraphs(const uint32 InNumFrames);
 
 	/** Called when the capture file has been fully loaded. */
 	void LoadComplete();
 
 	/** Sets number of frames. */
-	void SetNumberOfFrames( int32 InNumFrames );
+	void SetNumberOfFrames(int32 InNumFrames);
 	
 	/**
 	 * @return progress as floating point between 0 and 1.
@@ -1036,10 +1055,10 @@ protected:
 	FDelegateHandle OnTickHandle;
 
 	/** The data provider which holds all the collected profiler samples. */
-	IDataProviderRef DataProvider;
+	TSharedRef<IDataProvider> DataProvider;
 
 	/** The stat metadata which holds all collected stats descriptions. */
-	FProfilerStatMetaDataRef StatMetaData;
+	TSharedRef<FProfilerStatMetaData> StatMetaData;
 
 	/** Aggregated event graph data for all collected frames, used for generating average values. Also contains min and max. */
 	FEventGraphDataPtr EventGraphDataTotal;
@@ -1053,21 +1072,21 @@ protected:
 	/** Temporary event graph data for the specified frame. Recreated each frame, used by the task graph tasks only. */
 	const FEventGraphData* EventGraphDataCurrent;
 
-	/** Event graph completion sync ( combine for max, combine for add ) that can be done in parallel, but need to wait before combining the next frame. */
+	/** Event graph completion sync (combine for max, combine for add) that can be done in parallel, but need to wait before combining the next frame. */
 	FGraphEventRef CompletionSync;
 
 	/** Combined event graphs calculated on the task graph threads. */
 	TLockFreePointerListFIFO<FEventGraphData, PLATFORM_CACHE_LINE_SIZE> CombinedSubEventGraphsLFL;
 
 	// TODO: Temporary, need to be stored in metadata or send via UMB.
-	/** The time when this profiler session was created ( time of the connection to the client, time when a capture file was created ). */
+	/** The time when this profiler session was created (time of the connection to the client, time when a capture file was created). */
 	FDateTime CreationTime;
 
 	/** Session type for this profiler session. */
-	EProfilerSessionTypes::Type SessionType;
+	EProfilerSessionTypes SessionType;
 
 	/** Shared pointer to the session instance info. */
-	const ISessionInstanceInfoPtr SessionInstanceInfo;
+	const TSharedPtr<ISessionInstanceInfo> SessionInstanceInfo;
 
 	/** An unique session instance ID. */
 	FGuid SessionInstanceID;
@@ -1093,4 +1112,4 @@ protected:
 public:
 	/** Provides analysis of the frame rate */
 	TSharedRef<FFPSAnalyzer> FPSAnalyzer;
-};
\ No newline at end of file
+};
diff --git a/Engine/Source/Developer/Profiler/Private/ProfilerStream.cpp b/Engine/Source/Developer/Profiler/Private/ProfilerStream.cpp
--- a/Engine/Source/Developer/Profiler/Private/ProfilerStream.cpp
+++ b/Engine/Source/Developer/Profiler/Private/ProfilerStream.cpp
@@ -1,6 +1,8 @@
 // Copyright 1998-2016 Epic Games, Inc. All Rights Reserved.
 
 #include "ProfilerPrivatePCH.h"
+#include "ProfilerStream.h"
+
 
 /*-----------------------------------------------------------------------------
 	FProfilerUIStream
diff --git a/Engine/Source/Developer/Profiler/Private/Widgets/SDataGraph.cpp b/Engine/Source/Developer/Profiler/Private/Widgets/SDataGraph.cpp
--- a/Engine/Source/Developer/Profiler/Private/Widgets/SDataGraph.cpp
+++ b/Engine/Source/Developer/Profiler/Private/Widgets/SDataGraph.cpp
@@ -1,9 +1,17 @@
 // Copyright 1998-2016 Epic Games, Inc. All Rights Reserved.
 
 #include "ProfilerPrivatePCH.h"
+#include "ProfilerDataProvider.h"
+#include "ProfilerDataSource.h"
+#include "ProfilerManager.h"
+#include "ProfilerSession.h"
+#include "SDataGraph.h"
+#include "StatDragDropOp.h"
+
 
 #define LOCTEXT_NAMESPACE "SDataGraph"
 
+
 /*-----------------------------------------------------------------------------
 	STrackedStatSummary/SDataGraphSummary
 -----------------------------------------------------------------------------*/
@@ -26,7 +34,7 @@ public:
 		{}
 
 		SLATE_ARGUMENT( TSharedPtr<SDataGraph>, ParentWidget )
-		SLATE_ARGUMENT( FTrackedStatPtr, TrackedStat )
+		SLATE_ARGUMENT( TSharedPtr<FTrackedStat>, TrackedStat )
 		SLATE_EVENT( FGetHoveredFrameIndexDelegate, OnGetMouseFrameIndex )
 	SLATE_END_ARGS()
 
@@ -126,7 +134,7 @@ protected:
 		const bool bCanDisplayData = TrackedStat->GraphDataSource->CanBeDisplayedAsIndexBased() && ParentWidget.Pin()->GetViewMode() == EDataGraphViewModes::Index;
 
 		const uint32 FrameIndex = OnGetMouseFrameIndex.IsBound() ? (uint32)OnGetMouseFrameIndex.Execute() : 0;
-		const FGraphDataSourceRefConst& GraphDataSource = TrackedStat->GraphDataSource;
+		const TSharedRef<const FGraphDataSource>& GraphDataSource = TrackedStat->GraphDataSource;
 
 		if( bCanDisplayData && FrameIndex < GraphDataSource->GetNumFrames() )
 		{
@@ -156,7 +164,7 @@ protected:
 
 private:
 	/** A shared pointer to the traced stat. */
-	FTrackedStatPtr TrackedStat;
+	TSharedPtr<FTrackedStat> TrackedStat;
 
 	/** A weak pointer to the parent widget. */
 	TWeakPtr<SDataGraph> ParentWidget;
@@ -237,7 +245,7 @@ void SDataGraph::Tick( const FGeometry& AllottedGeometry, const double InCurrent
 
 void SDataGraph::UpdateState()
 {
-	const FTrackedStatPtr FirstTrackedStat = GetFirstGraph();
+	const TSharedPtr<FTrackedStat> FirstTrackedStat = GetFirstGraph();
 	if( FirstTrackedStat.IsValid() )
 	{
 		// If the view mode is index based, use the first source for reading the number of frames.
@@ -403,7 +411,7 @@ int32 SDataGraph::OnPaint( const FPaintArgs& Args, const FGeometry& AllottedGeom
 			{
 				// Draw line graph for each graph data source.
 
-				const FGraphDataSourceRefConst& GraphDataSource = TrackedStat.GraphDataSource;
+				const TSharedRef<const FGraphDataSource>& GraphDataSource = TrackedStat.GraphDataSource;
 
 				for (float GraphStartTimeMS = GraphOffsetMS; GraphStartTimeMS < GraphRangeEndMS; GraphStartTimeMS += TimeAccuracyMS)
 				{
@@ -433,7 +441,7 @@ int32 SDataGraph::OnPaint( const FPaintArgs& Args, const FGeometry& AllottedGeom
 		{
 			if( MultiMode == EDataGraphMultiModes::OneLinePerDataSource )
 			{						
-				const FGraphDataSourceRefConst& GraphDataSource = TrackedStat.GraphDataSource;
+				const TSharedRef<const FGraphDataSource>& GraphDataSource = TrackedStat.GraphDataSource;
 				const int32 GraphRangeEndIndex = FMath::Min( GraphOffset+NumVisiblePoints+1, NumDataPoints );
  
 				for( uint32 GraphStartIndex = (uint32)GraphOffset; GraphStartIndex < (uint32)GraphRangeEndIndex; GraphStartIndex++ )
@@ -482,11 +490,11 @@ int32 SDataGraph::OnPaint( const FPaintArgs& Args, const FGeometry& AllottedGeom
 	TArray<FVector2D> LinePoints;
 	const float LabelSize = MaxFontCharHeight * 7.0f;
 
-	const FTrackedStatPtr FirstTrackedStat = GetFirstGraph();
+	const TSharedPtr<FTrackedStat> FirstTrackedStat = GetFirstGraph();
 	if( FirstTrackedStat.IsValid() )
 	{
-		const FGraphDataSourceRefConst& GraphDataSource = FirstTrackedStat->GraphDataSource;
-		const IDataProviderRef DataProvider = GraphDataSource->GetDataProvider();
+		const TSharedRef<const FGraphDataSource>& GraphDataSource = FirstTrackedStat->GraphDataSource;
+		const TSharedRef<IDataProvider> DataProvider = GraphDataSource->GetDataProvider();
 
 		if( ViewMode == EDataGraphViewModes::Index )
 		{
@@ -958,7 +966,7 @@ int32 SDataGraph::OnPaint( const FPaintArgs& Args, const FGeometry& AllottedGeom
 
 
 
-void SDataGraph::AddInnerGraph( const FTrackedStatPtr& TrackedStat )
+void SDataGraph::AddInnerGraph( const TSharedPtr<FTrackedStat>& TrackedStat )
 {
 	TSharedPtr<SWidget> GraphSummary;
 	GraphDescriptionsVBox->AddSlot()
@@ -981,7 +989,7 @@ void SDataGraph::AddInnerGraph( const FTrackedStatPtr& TrackedStat )
 
 void SDataGraph::RemoveInnerGraph( const uint32 InStatID )
 {
-	FTrackedStatPtr TrackedStat = StatIDToGraphDescriptionMapping.FindRef( InStatID );
+	TSharedPtr<FTrackedStat> TrackedStat = StatIDToGraphDescriptionMapping.FindRef( InStatID );
 	if( TrackedStat.IsValid() )
 	{
 		TSharedRef<SWidget> DataGraphSummary = StatIDToWidgetMapping.FindChecked( InStatID );
@@ -1326,9 +1334,9 @@ bool SDataGraph::ViewMode_SetTimeBased_IsChecked() const
 	return ViewMode == EDataGraphViewModes::Time;
 }
 
-const FTrackedStatPtr SDataGraph::GetFirstGraph() const
+const TSharedPtr<FTrackedStat> SDataGraph::GetFirstGraph() const
 {
-	FTrackedStatPtr Result;
+	TSharedPtr<FTrackedStat> Result;
 	for (const auto& It : StatIDToGraphDescriptionMapping)
 	{
 		Result = It.Value; 
