Copying //UE4/Dev-Sequencer to //UE4/Dev-Main (Source: //UE4/Dev-Sequencer @ 3178529)

diff --git a/Engine/Source/Developer/Profiler/Private/Widgets/SDataGraph.h b/Engine/Source/Developer/Profiler/Private/Widgets/SDataGraph.h
--- a/Engine/Source/Developer/Profiler/Private/Widgets/SDataGraph.h
+++ b/Engine/Source/Developer/Profiler/Private/Widgets/SDataGraph.h
@@ -2,9 +2,11 @@
 
 #pragma once
 
-/*-----------------------------------------------------------------------------
-	Enumerators
------------------------------------------------------------------------------*/
+#include "ProfilerDataSource.h"
+
+
+class FTrackedStat;
+
 
 namespace EDataGraphViewModes
 {
@@ -21,6 +23,7 @@ namespace EDataGraphViewModes
 	};
 }
 
+
 namespace EDataGraphMultiModes
 {
 	enum Type
@@ -35,9 +38,6 @@ namespace EDataGraphMultiModes
 	};
 }
 
-/*-----------------------------------------------------------------------------
-	Delegates
------------------------------------------------------------------------------*/
 
 /** The delegate to be invoked when the frame offset has been changed. */
 DECLARE_DELEGATE_OneParam( FGraphOffsetChangedDelegate, int /*InFrameOffset*/ );
@@ -54,10 +54,6 @@ DECLARE_DELEGATE_OneParam( FViewModeChangedDelegate, EDataGraphViewModes::Type )
 DECLARE_DELEGATE_TwoParams( FSelectionChangedForTimeDelegate, float /*FrameStartTimeMS*/, float /*FrameEndTimeMS*/ );
 
 
-/*-----------------------------------------------------------------------------
-	Declarations
------------------------------------------------------------------------------*/
-
 /** A custom widget used to display graphs. */
 // Rename to SDataChart ?
 class SDataGraph : public SCompoundWidget
@@ -238,7 +234,7 @@ protected:
 	FSelectionChangedForIndexEvent SelectionChangedForIndexEvent;
 
 public:
-	void AddInnerGraph( const FTrackedStatPtr& TrackedStat );
+	void AddInnerGraph( const TSharedPtr<FTrackedStat>& TrackedStat );
 	void RemoveInnerGraph( const uint32 StatID );
 
 	/**
@@ -313,7 +309,7 @@ protected:
 	/** Handles FIsActionChecked for DataGraph_ViewMode_SetTimeBased. */
 	bool ViewMode_SetTimeBased_IsChecked() const;
 
-	const FTrackedStatPtr GetFirstGraph() const;
+	const TSharedPtr<FTrackedStat> GetFirstGraph() const;
 
 	/*-----------------------------------------------------------------------------
 		Misc
@@ -322,7 +318,7 @@ protected:
 	/** Horizontal box widget where graph descriptions are displayed. */
 	TSharedPtr<SVerticalBox> GraphDescriptionsVBox;
 
-	TMap< uint32, FTrackedStatPtr > StatIDToGraphDescriptionMapping;
+	TMap< uint32, TSharedPtr<FTrackedStat> > StatIDToGraphDescriptionMapping;
 	TMap< uint32, TSharedRef<SWidget> > StatIDToWidgetMapping;
 
 	/** The current mouse position. */
diff --git a/Engine/Source/Developer/Profiler/Private/Widgets/SEventGraph.cpp b/Engine/Source/Developer/Profiler/Private/Widgets/SEventGraph.cpp
--- a/Engine/Source/Developer/Profiler/Private/Widgets/SEventGraph.cpp
+++ b/Engine/Source/Developer/Profiler/Private/Widgets/SEventGraph.cpp
@@ -1,11 +1,15 @@
 // Copyright 1998-2016 Epic Games, Inc. All Rights Reserved.
 
 #include "ProfilerPrivatePCH.h"
+#include "ProfilerManager.h"
+#include "SEventGraph.h"
 #include "SEventGraphTooltip.h"
-#include "SSearchBox.h"
+#include "StatDragDropOp.h"
+
 
 #define LOCTEXT_NAMESPACE "SEventGraph"
 
+
 namespace EEventGraphViewModes
 {
 	FText ToName( const Type EventGraphViewMode )
@@ -2429,7 +2433,7 @@ void SEventGraph::ShowEventsInViewMode( const TArray<FEventGraphSamplePtr>& Even
 	SetHierarchicalSelectedEvents( EventGraphState->SelectedEvents );
 	SetHierarchicalExpandedEvents( EventGraphState->ExpandedEvents );
 		
-	EEventPropertyIndex::Type ColumnIndex = FEventGraphColumns::Get().ColumnNameToIndexMapping.FindChecked( ColumnBeingSorted )->Index;
+	EEventPropertyIndex ColumnIndex = FEventGraphColumns::Get().ColumnNameToIndexMapping.FindChecked( ColumnBeingSorted )->Index;
 
 	if( NewViewMode == EEventGraphViewModes::FlatInclusive || NewViewMode == EEventGraphViewModes::FlatInclusiveCoalesced || NewViewMode == EEventGraphViewModes::Hierarchical )
 	{
@@ -2464,7 +2468,7 @@ void SEventGraph::ShowEventsInViewMode( const TArray<FEventGraphSamplePtr>& Even
 	TreeView_Refresh();
 }
 
-void SEventGraph::ScrollToTheSlowestSelectedEvent( EEventPropertyIndex::Type ColumnIndex )
+void SEventGraph::ScrollToTheSlowestSelectedEvent( EEventPropertyIndex ColumnIndex )
 {
 	TArray<FEventGraphSamplePtr> SelectedEvents = TreeView_Base->GetSelectedItems();
 	if( SelectedEvents.Num() > 0 )
@@ -3282,7 +3286,7 @@ bool SEventGraph::SelectAllFrames_CanExecute() const
 	return IsEventGraphStatesHistoryValid();
 }
 
-void SEventGraph::ProfilerManager_OnViewModeChanged( EProfilerViewMode::Type NewViewMode )
+void SEventGraph::ProfilerManager_OnViewModeChanged( EProfilerViewMode NewViewMode )
 {
 // 	if( NewViewMode == EProfilerViewMode::LineIndexBased )
 // 	{
diff --git a/Engine/Source/Developer/Profiler/Private/Widgets/SEventGraph.h b/Engine/Source/Developer/Profiler/Private/Widgets/SEventGraph.h
--- a/Engine/Source/Developer/Profiler/Private/Widgets/SEventGraph.h
+++ b/Engine/Source/Developer/Profiler/Private/Widgets/SEventGraph.h
@@ -2,9 +2,11 @@
 
 #pragma once
 
-/*-----------------------------------------------------------------------------
-	Type definitions
------------------------------------------------------------------------------*/
+#include "ProfilerDataSource.h"
+
+
+enum class EProfilerViewMode;
+
 
 /** Type definition for shared pointers to instances of FEventGraphColumn. */
 typedef TSharedPtr<class FEventGraphColumn> FEventGraphColumnPtr;
@@ -12,9 +14,7 @@ typedef TSharedPtr<class FEventGraphColumn> FEventGraphColumnPtr;
 /** Type definition for shared references to instances of FEventGraphColumn. */
 typedef TSharedRef<class FEventGraphColumn> FEventGraphColumnRef;
 
-/*-----------------------------------------------------------------------------
-	Enumerators
------------------------------------------------------------------------------*/
+
 
 /** Enumerates event graph view modes. */
 namespace EEventGraphViewModes
@@ -60,6 +60,7 @@ namespace EEventGraphViewModes
 	FName ToBrushName( const Type EventGraphViewMode );
 }
 
+
 /** Holds information about a column in the event graph widget. */
 class FEventGraphColumn
 {
@@ -73,7 +74,7 @@ protected:
 	/** Initialization constructor, only used in FEventGraphColumns. */
 	FEventGraphColumn
 	(
-		const EEventPropertyIndex::Type InIndex,
+		const EEventPropertyIndex InIndex,
 		const FName InSearchID,
 		const FText InShortName,
 		const FText InDescription,
@@ -101,7 +102,7 @@ protected:
 
 public:
 	/** Index of the event's property, also means the index of the column @see EEventPropertyIndex. */
-	EEventPropertyIndex::Type Index;
+	EEventPropertyIndex Index;
 
 	/** Name of the column, name of the property. */
 	FName ID;
@@ -216,7 +217,7 @@ public:
 	}
 
 protected:
-	void ProfilerManager_OnViewModeChanged( EProfilerViewMode::Type NewViewMode );
+	void ProfilerManager_OnViewModeChanged( EProfilerViewMode NewViewMode );
 	
 protected:
 	/** The event to execute when the event graph has been restored from the history. */
@@ -244,7 +245,7 @@ protected:
 	void GetHierarchicalExpandedEvents( TSet<FEventGraphSamplePtr>& out_HierarchicalExpandedEvents ) const;
 	
 	void ShowEventsInViewMode( const TArray<FEventGraphSamplePtr>& EventsToSynchronize, const EEventGraphViewModes::Type NewViewMode );
-	void ScrollToTheSlowestSelectedEvent( EEventPropertyIndex::Type ColumnIndex );
+	void ScrollToTheSlowestSelectedEvent( EEventPropertyIndex ColumnIndex );
 
 	void CreateEvents();
 	void SortEvents();
diff --git a/Engine/Source/Developer/Profiler/Private/Widgets/SEventGraphTooltip.cpp b/Engine/Source/Developer/Profiler/Private/Widgets/SEventGraphTooltip.cpp
--- a/Engine/Source/Developer/Profiler/Private/Widgets/SEventGraphTooltip.cpp
+++ b/Engine/Source/Developer/Profiler/Private/Widgets/SEventGraphTooltip.cpp
@@ -1,12 +1,15 @@
 // Copyright 1998-2016 Epic Games, Inc. All Rights Reserved.
 
 #include "ProfilerPrivatePCH.h"
+#include "ProfilerDataSource.h"
 #include "SEventGraphTooltip.h"
 
+
 #define LOCTEXT_NAMESPACE "SEventGraphTooltip"
 
+
 BEGIN_SLATE_FUNCTION_BUILD_OPTIMIZATION
-TSharedPtr<SToolTip> SEventGraphTooltip::GetTableCellTooltip( const FEventGraphSamplePtr EventSample )
+TSharedPtr<SToolTip> SEventGraphTooltip::GetTableCellTooltip( const TSharedPtr<FEventGraphSample> EventSample )
 {
 	const FSlateFontInfo TitleFont( FPaths::EngineContentDir() / TEXT("Slate/Fonts/Roboto-Bold.ttf"), 10 );
 	const FSlateFontInfo DescriptionFont( FPaths::EngineContentDir() / TEXT("Slate/Fonts/Roboto-Regular.ttf"), 8 );
@@ -492,4 +495,12 @@ TSharedPtr<SToolTip> SEventGraphTooltip::GetTableCellTooltip( const FEventGraphS
 }
 END_SLATE_FUNCTION_BUILD_OPTIMIZATION
 
+
+EVisibility SEventGraphTooltip::GetHotPathIconVisibility(const TSharedPtr<FEventGraphSample> EventSample)
+{
+	const bool bIsHotPathIconVisible = EventSample->_bIsHotPath;
+	return bIsHotPathIconVisible ? EVisibility::Visible : EVisibility::Collapsed;
+}
+
+
 #undef LOCTEXT_NAMESPACE
\ No newline at end of file
diff --git a/Engine/Source/Developer/Profiler/Private/Widgets/SEventGraphTooltip.h b/Engine/Source/Developer/Profiler/Private/Widgets/SEventGraphTooltip.h
--- a/Engine/Source/Developer/Profiler/Private/Widgets/SEventGraphTooltip.h
+++ b/Engine/Source/Developer/Profiler/Private/Widgets/SEventGraphTooltip.h
@@ -2,17 +2,19 @@
 
 #pragma once
 
-/** An advanced tooltip used to show various informations in the event graph widget. */
+class FEventGraphSample;
+
+
+/**
+ * An advanced tooltip used to show various informations in the event graph widget.
+ */
 class SEventGraphTooltip
 {
-protected:
-	static EVisibility GetHotPathIconVisibility( const FEventGraphSamplePtr EventSample )
-	{
-		const bool bIsHotPathIconVisible = EventSample->_bIsHotPath;
-		return bIsHotPathIconVisible ? EVisibility::Visible : EVisibility::Collapsed;
-	}
-
 public:
-	static TSharedPtr<SToolTip> GetTableCellTooltip( const FEventGraphSamplePtr EventSample );
 
-};
\ No newline at end of file
+	static TSharedPtr<SToolTip> GetTableCellTooltip( const TSharedPtr<FEventGraphSample> EventSample );
+
+protected:
+
+	static EVisibility GetHotPathIconVisibility(const TSharedPtr<FEventGraphSample> EventSample);
+};
diff --git a/Engine/Source/Developer/Profiler/Private/Widgets/SFiltersAndPresets.cpp b/Engine/Source/Developer/Profiler/Private/Widgets/SFiltersAndPresets.cpp
--- a/Engine/Source/Developer/Profiler/Private/Widgets/SFiltersAndPresets.cpp
+++ b/Engine/Source/Developer/Profiler/Private/Widgets/SFiltersAndPresets.cpp
@@ -1,10 +1,15 @@
 // Copyright 1998-2016 Epic Games, Inc. All Rights Reserved.
 
 #include "ProfilerPrivatePCH.h"
-#include "SSearchBox.h"
+#include "ProfilerManager.h"
+#include "ProfilerSession.h"
+#include "SFiltersAndPresets.h"
+#include "StatDragDropOp.h"
+
 
 #define LOCTEXT_NAMESPACE "SFiltersAndPresets"
 
+
 struct SFiltersAndPresetsHelper
 {
 	static const FSlateBrush* GetIconForGroup()
@@ -38,7 +43,8 @@ struct SFiltersAndPresetsHelper
 class SFiltersAndPresetsTooltip
 {
 	const uint32 StatID;
-	FProfilerSessionPtr ProfilerSession;
+	TSharedPtr<FProfilerSession> ProfilerSession;
+
 public:
 	SFiltersAndPresetsTooltip( const uint32 InStatID )
 		: StatID( InStatID )
@@ -758,11 +764,11 @@ END_SLATE_FUNCTION_BUILD_OPTIMIZATION
 
 void SFiltersAndPresets::ProfilerManager_OnRequestFilterAndPresetsUpdate()
 {
-	FProfilerSessionPtr ProfilerSessionLocal = FProfilerManager::Get()->GetProfilerSession();
+	TSharedPtr<FProfilerSession> ProfilerSessionLocal = FProfilerManager::Get()->GetProfilerSession();
 	UpdateGroupAndStatTree( ProfilerSessionLocal );
 }
 
-void SFiltersAndPresets::UpdateGroupAndStatTree( const FProfilerSessionPtr InProfilerSession )
+void SFiltersAndPresets::UpdateGroupAndStatTree( const TSharedPtr<FProfilerSession> InProfilerSession )
 {
 	const bool bRebuild = InProfilerSession != ProfilerSession;
 	if( bRebuild )
@@ -772,7 +778,7 @@ void SFiltersAndPresets::UpdateGroupAndStatTree( const FProfilerSessionPtr InPro
 
 	ProfilerSession = InProfilerSession;
 
-	const FProfilerStatMetaDataRef StatMetaData = ProfilerSession->GetMetaData();
+	const TSharedRef<FProfilerStatMetaData> StatMetaData = ProfilerSession->GetMetaData();
 
 	// Create all stat nodes.
 	for( auto It = StatMetaData->GetStatIterator(); It; ++It )
diff --git a/Engine/Source/Developer/Profiler/Private/Widgets/SFiltersAndPresets.h b/Engine/Source/Developer/Profiler/Private/Widgets/SFiltersAndPresets.h
--- a/Engine/Source/Developer/Profiler/Private/Widgets/SFiltersAndPresets.h
+++ b/Engine/Source/Developer/Profiler/Private/Widgets/SFiltersAndPresets.h
@@ -2,9 +2,13 @@
 
 #pragma once
 
-/*-----------------------------------------------------------------------------
-	Type definitions
------------------------------------------------------------------------------*/
+#include "ProfilerSample.h"
+#include "FilterCollection.h"
+#include "TextFilter.h"
+
+
+class FProfilerSession;
+
 
 /** Type definition for shared pointers to instances of FGroupOrStatNode. */
 typedef TSharedPtr<class FGroupOrStatNode> FGroupOrStatNodePtr;
@@ -18,9 +22,6 @@ typedef TSharedRef<const class FGroupOrStatNode> FGroupOrStatNodeRefConst;
 /** Type definition for weak references to instances of FGroupOrStatNode. */
 typedef TWeakPtr<class FGroupOrStatNode> FGroupOrStatNodeWeak;
 
-/*-----------------------------------------------------------------------------
-	Enumerators
------------------------------------------------------------------------------*/
 
 /** Enumerates types of grouping or sorting for the group and stat nodes. */
 struct EStatGroupingOrSortingMode
@@ -62,9 +63,6 @@ struct EStatGroupingOrSortingMode
 	static FName ToBrushName( const Type StatGroupingOrSortingMode );
 };
 
-/*-----------------------------------------------------------------------------
-	Declarations
------------------------------------------------------------------------------*/
 
 /** The group and stat filter collection - used for updating the list of groups and stats. */
 typedef TFilterCollection< const FGroupOrStatNodePtr& > FGroupAndStatFilterCollection;
@@ -72,10 +70,12 @@ typedef TFilterCollection< const FGroupOrStatNodePtr& > FGroupAndStatFilterColle
 /** The group and stat text based filter - used for updating the list of groups and stats. */
 typedef TTextFilter< const FGroupOrStatNodePtr& > FGroupAndStatTextFilter;
 
-// type, value, filtering, etc
 
-/** Class used to store information about stat and group metadata. */
-class FGroupOrStatNode : public TSharedFromThis<FGroupOrStatNode>
+/**
+ * Class used to store information about stat and group metadata.
+ */
+class FGroupOrStatNode
+	: public TSharedFromThis<FGroupOrStatNode>
 {
 public:
 	/** Initialization constructor for the stat node. */
@@ -223,6 +223,7 @@ public:
 	bool bForceExpandGroupNode;
 };
 
+
 /** Helper struct that contains sorting classes. */
 struct FGroupAndStatSorting
 {
@@ -265,8 +266,10 @@ struct FGroupAndStatSorting
 	};
 };
 
+
 /** Configurable window with advanced options for filtering and creating presets. */
-class SFiltersAndPresets : public SCompoundWidget
+class SFiltersAndPresets
+	: public SCompoundWidget
 {
 public:
 	/** Default constructor. */
@@ -288,7 +291,7 @@ protected:
 	/** Called when the filter and presets widget should be updated with the latest data. */
 	void ProfilerManager_OnRequestFilterAndPresetsUpdate();
 
-	void UpdateGroupAndStatTree( const FProfilerSessionPtr InProfilerSession );
+	void UpdateGroupAndStatTree( const TSharedPtr<FProfilerSession> InProfilerSession );
 	void CreateGroups();
 	void SortStats();
 
@@ -398,7 +401,7 @@ protected:
 	TSharedPtr< FGroupAndStatFilterCollection > GroupAndStatFilters;
 	
 	/** A weak pointer to the profiler session used to populate this widget. */
-	FProfilerSessionPtr/*Weak*/ ProfilerSession;
+	TSharedPtr<FProfilerSession>/*Weak*/ ProfilerSession;
 
 	/** How we group the metadata?. */
 	EStatGroupingOrSortingMode::Type GroupingMode;
diff --git a/Engine/Source/Developer/Profiler/Private/Widgets/SHistogram.cpp b/Engine/Source/Developer/Profiler/Private/Widgets/SHistogram.cpp
--- a/Engine/Source/Developer/Profiler/Private/Widgets/SHistogram.cpp
+++ b/Engine/Source/Developer/Profiler/Private/Widgets/SHistogram.cpp
@@ -1,9 +1,13 @@
 // Copyright 1998-2016 Epic Games, Inc. All Rights Reserved.
 
 #include "ProfilerPrivatePCH.h"
+#include "ProfilerFPSAnalyzer.h"
+#include "SHistogram.h"
+
 
 #define LOCTEXT_NAMESPACE "SHistogram"
 
+
 /*-----------------------------------------------------------------------------
 	SHistogram
 -----------------------------------------------------------------------------*/
diff --git a/Engine/Source/Developer/Profiler/Private/Widgets/SHistogram.h b/Engine/Source/Developer/Profiler/Private/Widgets/SHistogram.h
--- a/Engine/Source/Developer/Profiler/Private/Widgets/SHistogram.h
+++ b/Engine/Source/Developer/Profiler/Private/Widgets/SHistogram.h
@@ -2,9 +2,8 @@
 
 #pragma once
 
-/*-----------------------------------------------------------------------------
-	Type definitions
------------------------------------------------------------------------------*/
+#include "ProfilerSample.h"
+
 
 /** Type definition for shared pointers to instances of FGraphDataSource. */
 typedef TSharedPtr<class IHistogramDataSource> FHistogramDataSourcePtr;
@@ -12,9 +11,7 @@ typedef TSharedPtr<class IHistogramDataSource> FHistogramDataSourcePtr;
 /** Type definition for shared references to instances of FGraphDataSource. */
 typedef TSharedRef<class IHistogramDataSource> FHistogramDataSourceRef;
 
-/*-----------------------------------------------------------------------------
-	Basic structures
------------------------------------------------------------------------------*/
+
 class FHistogramDescription
 {
 public:
@@ -22,13 +19,13 @@ public:
 	{}
 
 	/**
-	 * Constructor
+	 * Create and initialize a new instance.
 	 *
-	 * @param InDataSource   - data source for the histogram
-	 * @param InBinInterval  - interval for each bin
-	 * @param InMinValue     - minimum value to use for the bins
-	 * @param InMaxValue     - maximum value to use for the bins
-	 * @param InBinNormalize - whether or not to normalize the bins to the total data count
+	 * @param InDataSource Data source for the histogram.
+	 * @param InBinInterval Interval for each bin.
+	 * @param InMinValue Minimum value to use for the bins.
+	 * @param InMaxValue Maximum value to use for the bins.
+	 * @param InBinNormalize Whether or not to normalize the bins to the total data count.
 	 *
 	 */
 	FHistogramDescription(const FHistogramDataSourceRef& InDataSource, float InBinInterval, float InMinValue, float InMaxValue, bool InBinNormalize = false)
diff --git a/Engine/Source/Developer/Profiler/Private/Widgets/SMultiDumpBrowser.cpp b/Engine/Source/Developer/Profiler/Private/Widgets/SMultiDumpBrowser.cpp
new file mode 100644
--- /dev/null
+++ b/Engine/Source/Developer/Profiler/Private/Widgets/SMultiDumpBrowser.cpp
@@ -0,0 +1,202 @@
+// Copyright 1998-2016 Epic Games, Inc. All Rights Reserved.
+
+#include "ProfilerPrivatePCH.h"
+#include "ProfilerDataSource.h"
+#include "ProfilerManager.h"
+#include "ProfilerSession.h"
+#include "SMultiDumpBrowser.h"
+
+
+#define LOCTEXT_NAMESPACE "SMultiDumpBrowser"
+
+
+void SMultiDumpBrowser::Construct(const FArguments& InArgs)
+{
+	this->SetEnabled(true);
+
+	ChildSlot
+	[
+		SNew(SOverlay)
+		+ SOverlay::Slot()
+		.Padding(4)
+		.VAlign(VAlign_Fill)
+		[
+			SNew(SVerticalBox)
+			+ SVerticalBox::Slot().AutoHeight().Padding(2)
+			[
+				SNew(SBorder).BorderImage(FEditorStyle::GetBrush("ToolPanel.GroupBorder")).Padding(4).HAlign(HAlign_Fill)
+				[
+					SNew(SHorizontalBox)
+					+ SHorizontalBox::Slot()
+						[
+							SNew(STextBlock).Text(LOCTEXT("MultiDumpBrowserThreadTotals", "Show thread totals for:"))
+						]
+					+ SHorizontalBox::Slot().HAlign(HAlign_Fill)
+						[
+							SAssignNew(DisplayTotalsBox, SEditableTextBox).OnTextCommitted(this, &SMultiDumpBrowser::PrefilterTextCommitted).ToolTipText(LOCTEXT("MultiDumpBrowserTooltip", "Use \"Load Folder\" above to load a folder of stats dumps. GT/RT totals for stats matching text entered here will be displayed in the list below - e.g. enter \"Particle\" here to show total thread times for particle emitters."))
+						]
+				]
+			]
+			+ SVerticalBox::Slot().Padding(2).FillHeight(1.0f)
+			[
+				SNew(SBorder).BorderImage(FEditorStyle::GetBrush("ToolPanel.GroupBorder")).Padding(4).HAlign(HAlign_Fill).VAlign(VAlign_Fill)
+				[
+					SAssignNew(FileList, SListView<TSharedPtr<FFileDescriptor>>)
+					.ListItemsSource(&StatsFiles)
+					.ItemHeight(16)
+					.OnGenerateRow(this, &SMultiDumpBrowser::GenerateFileRow)
+					.OnSelectionChanged(this, &SMultiDumpBrowser::SelectionChanged)
+				]
+			]
+		]
+	];
+}
+
+
+TSharedRef<ITableRow> SMultiDumpBrowser::GenerateFileRow(TSharedPtr<FFileDescriptor> FileItem, const TSharedRef<STableViewBase>& OwnerTable)
+{
+	return SNew(SFileTableRow, OwnerTable, FileItem);
+}
+
+
+void SMultiDumpBrowser::SelectionChanged(TSharedPtr<FFileDescriptor> SelectedItem, ESelectInfo::Type SelType)
+{
+	FProfilerManager::Get()->LoadProfilerCapture(SelectedItem->FullPath);
+}
+
+
+void SMultiDumpBrowser::Update()
+{
+	FileList->RequestListRefresh();
+}
+
+
+void SMultiDumpBrowser::AddFile(FFileDescriptor *InFileDesc)
+{
+	StatsFiles.Add(MakeShareable(InFileDesc));
+}
+
+
+void SMultiDumpBrowser::Clear()
+{
+	StatsFiles.Empty();
+}
+
+
+void SMultiDumpBrowser::Tick(const FGeometry& AllottedGeometry, const double InCurrentTime, const float InDeltaTime)
+{
+	static bool bCurrentlyLoading = false;
+
+	// iterate over all stats files, load one at a time, analyze
+	for (auto Desc : StatsFiles)
+	{
+		// if we're not currently loading one and we've found one in the list that was added but not analyzed, load it
+		if (Desc->Status == FFileDescriptor::EQueued && bCurrentlyLoading == false)
+		{
+			FProfilerManager::Get()->LoadProfilerCapture(Desc->FullPath);
+			Desc->Status = FFileDescriptor::ELoading;
+			Desc->TimeString = "Getting timings, please wait...";
+			bCurrentlyLoading = true;
+		}
+
+		// if the current load is completed, start summing up thread totals for the term we're looking for in the DisplayTotalsBox
+		if (FProfilerManager::Get()->IsCaptureFileFullyProcessed() && Desc->Status == FFileDescriptor::ELoading)
+		{
+			Desc->Status = FFileDescriptor::ELoaded;
+			bCurrentlyLoading = false;
+			FindTotalsForPrefilter(Desc);
+
+			Desc->Status = FFileDescriptor::EAnalyzed;
+		}
+	}
+
+	if (bCurrentlyLoading == true)
+	{
+		this->SetEnabled(false);
+	}
+	else
+	{
+		this->SetEnabled(true);
+	}
+
+	SWidget::Tick(AllottedGeometry, InCurrentTime, InDeltaTime);
+}
+
+
+void SMultiDumpBrowser::FindTotalsForPrefilterRecursive(TSharedPtr<FEventGraphSample> Event, float &OutTotalTime)
+{
+	FString EventName = Event->_StatName.ToString().ToUpper();
+	if (EventName.Contains(TotalsFilteringText.ToUpper()))
+	{
+		OutTotalTime += Event->_InclusiveTimeMS;
+	}
+	else
+	{
+		for (auto Child : Event->GetChildren())
+		{
+			FindTotalsForPrefilterRecursive(Child, OutTotalTime);
+		}
+	}
+}
+
+
+void SMultiDumpBrowser::FindTotalsForPrefilter(TSharedPtr<FFileDescriptor> Desc)
+{
+	float TotalRenderThreadTime = 0;
+	float TotalGameThreadTime = 0;
+
+	// if we don't have a totals filtering text, display the entire RT and GT times
+	bool bGetBaseThreadTimes = false;
+	if (TotalsFilteringText == "")
+	{
+		bGetBaseThreadTimes = true;
+	}
+
+	FEventGraphSamplePtr RootPtr = FProfilerManager::Get()->GetProfilerSession()->GetEventGraphDataAverage()->GetRoot();
+
+	// start at the root's children, which are the threads
+	for (auto Child : RootPtr->GetChildren())
+	{
+		TArray<FString>EventNameTokens;
+		Child->_ThreadName.ToString().ParseIntoArray(EventNameTokens, TEXT(" "));
+
+		// get cumulative times for stats on the render thread
+		if (EventNameTokens[0] == "RenderThread")
+		{
+			if (bGetBaseThreadTimes)
+			{
+				TotalsFilteringText = "RenderThread";
+			}
+			FindTotalsForPrefilterRecursive(Child, TotalRenderThreadTime);
+		}
+		// get cumulative times for stats on the game thread
+		else if (EventNameTokens[0] == "GameThread")
+		{
+			if (bGetBaseThreadTimes)
+			{
+				TotalsFilteringText = "GameThread";
+			}
+			FindTotalsForPrefilterRecursive(Child, TotalGameThreadTime);
+		}
+
+		if (bGetBaseThreadTimes)
+		{
+			TotalsFilteringText = "";
+		}
+	}
+
+	Desc->TimeString = "RT " + FString::Printf(TEXT("%.2f"), TotalRenderThreadTime) + " / GT " + FString::Printf(TEXT("%.2f"), TotalGameThreadTime);
+}
+
+
+void SMultiDumpBrowser::PrefilterTextCommitted(const FText& InText, const ETextCommit::Type InTextAction)
+{
+	TotalsFilteringText = InText.ToString();
+	for (auto Desc : StatsFiles)
+	{
+		Desc->Status = FFileDescriptor::EQueued;
+	}
+}
+
+
+#undef LOCTEXT_NAMESPACE
diff --git a/Engine/Source/Developer/Profiler/Private/Widgets/SMultiDumpBrowser.h b/Engine/Source/Developer/Profiler/Private/Widgets/SMultiDumpBrowser.h
--- a/Engine/Source/Developer/Profiler/Private/Widgets/SMultiDumpBrowser.h
+++ b/Engine/Source/Developer/Profiler/Private/Widgets/SMultiDumpBrowser.h
@@ -3,21 +3,22 @@
 #pragma once
 
 
-#define LOCTEXT_NAMESPACE "SMultiDumpBrowser"
+class FEventGraphSample;
 
-/*-----------------------------------------------------------------------------
-Declarations
------------------------------------------------------------------------------*/
 
-
-/** A custom widget used to display a histogram. */
-class SMultiDumpBrowser : public SCompoundWidget
+/**
+ * A custom widget used to display a histogram.
+ */
+class SMultiDumpBrowser
+	: public SCompoundWidget
 {
 public:
 	SLATE_BEGIN_ARGS(SMultiDumpBrowser) { }
 	SLATE_END_ARGS()
 
-	// descriptor for a single stats file; associated with a table row, so we can load on demand
+	/**
+	 * Descriptor for a single stats file; associated with a table row, so we can load on demand.
+	 */
 	class FFileDescriptor
 	{
 	public:
@@ -43,8 +44,12 @@ public:
 		}
 	};
 
-	// Table row for the stats dump file list view
-	class SFileTableRow : public STableRow<TSharedPtr<FFileDescriptor>>
+
+	/**
+	 * Table row for the stats dump file list view.
+	 */
+	class SFileTableRow
+		: public STableRow<TSharedPtr<FFileDescriptor>>
 	{
 	public:
 		SLATE_BEGIN_ARGS(SFileTableRow) { }
@@ -77,210 +82,48 @@ public:
 
 
 	/**
-	* Construct this widget
+	* Construct this widget.
 	*
-	* @param	InArgs	The declaration data for this widget
-	*/
-	void Construct(const FArguments& InArgs)
-	{
-		this->SetEnabled(true);
-		ChildSlot
-			[
-				SNew(SOverlay)
-				+ SOverlay::Slot()
-				.Padding(4)
-				.VAlign(VAlign_Fill)
-				[
-					SNew(SVerticalBox)
-					+ SVerticalBox::Slot().AutoHeight().Padding(2)
-					[
-						SNew(SBorder).BorderImage(FEditorStyle::GetBrush("ToolPanel.GroupBorder")).Padding(4).HAlign(HAlign_Fill)
-						[
-							SNew(SHorizontalBox)
-							+ SHorizontalBox::Slot()
-								[
-									SNew(STextBlock).Text(LOCTEXT("MultiDumpBrowserThreadTotals", "Show thread totals for:"))
-								]
-							+ SHorizontalBox::Slot().HAlign(HAlign_Fill)
-								[
-									SAssignNew(DisplayTotalsBox, SEditableTextBox).OnTextCommitted(this, &SMultiDumpBrowser::PrefilterTextCommitted).ToolTipText(LOCTEXT("MultiDumpBrowserTooltip", "Use \"Load Folder\" above to load a folder of stats dumps. GT/RT totals for stats matching text entered here will be displayed in the list below - e.g. enter \"Particle\" here to show total thread times for particle emitters."))
-								]
-						]
-					]
-					+ SVerticalBox::Slot().Padding(2).FillHeight(1.0f)
-					[
-						SNew(SBorder).BorderImage(FEditorStyle::GetBrush("ToolPanel.GroupBorder")).Padding(4).HAlign(HAlign_Fill).VAlign(VAlign_Fill)
-						[
-							SAssignNew(FileList, SListView<TSharedPtr<FFileDescriptor>>)
-							.ListItemsSource(&StatsFiles)
-							.ItemHeight(16)
-							.OnGenerateRow(this, &SMultiDumpBrowser::GenerateFileRow)
-							.OnSelectionChanged(this, &SMultiDumpBrowser::SelectionChanged)
-						]
-					]
-				]
-			];
-	}
-
-	TSharedRef<ITableRow> GenerateFileRow(TSharedPtr<FFileDescriptor> FileItem, const TSharedRef<STableViewBase>& OwnerTable)
-	{
-		return SNew(SFileTableRow, OwnerTable, FileItem);
-	}
-
-	void SelectionChanged(TSharedPtr<FFileDescriptor> SelectedItem, ESelectInfo::Type SelType)
-	{
-		FProfilerManager::Get()->LoadProfilerCapture(SelectedItem->FullPath);
-	}
-
-
-	void Update()
-	{
-		FileList->RequestListRefresh();
-	}
-
-	void AddFile(FFileDescriptor *InFileDesc)
-	{
-		StatsFiles.Add(MakeShareable(InFileDesc));
-	}
-
-	void Clear()
-	{
-		StatsFiles.Empty();
-	}
-
-
-	/**
-	* Tick the file browser; we update the thread totals we've filtered for here; this is necessary because 
-	* the profiler window widgets process loading in their tick and we can't load more than a single file at a time, so we need to
-	* trigger loading of each one, then analyze in order to display thread totals for specific search terms
+	* @param InArgs The declaration data for this widget.
 	*/
-	virtual void Tick(const FGeometry& AllottedGeometry, const double InCurrentTime, const float InDeltaTime) override
-	{
-		static bool bCurrentlyLoading = false;
-
-		// iterate over all stats files, load one at a time, analyze
-		for (auto Desc : StatsFiles)
-		{
-			// if we're not currently loading one and we've found one in the list that was added but not analyzed, load it
-			if (Desc->Status == FFileDescriptor::EQueued && bCurrentlyLoading == false)
-			{
-				FProfilerManager::Get()->LoadProfilerCapture(Desc->FullPath);
-				Desc->Status = FFileDescriptor::ELoading;
-				Desc->TimeString = "Getting timings, please wait...";
-				bCurrentlyLoading = true;
-			}
-
-			// if the current load is completed, start summing up thread totals for the term we're looking for in the DisplayTotalsBox
-			if (FProfilerManager::Get()->IsCaptureFileFullyProcessed() && Desc->Status == FFileDescriptor::ELoading)
-			{
-				Desc->Status = FFileDescriptor::ELoaded;
-				bCurrentlyLoading = false;
-				FindTotalsForPrefilter(Desc);
-
-				Desc->Status = FFileDescriptor::EAnalyzed;
-			}
-		}
-
-		if (bCurrentlyLoading == true)
-		{
-			this->SetEnabled(false);
-		}
-		else
-		{
-			this->SetEnabled(true);
-		}
-
-		SWidget::Tick(AllottedGeometry, InCurrentTime, InDeltaTime);
-	}
+	void Construct(const FArguments& InArgs);
 
+	TSharedRef<ITableRow> GenerateFileRow(TSharedPtr<FFileDescriptor> FileItem, const TSharedRef<STableViewBase>& OwnerTable);
+	
+	void SelectionChanged(TSharedPtr<FFileDescriptor> SelectedItem, ESelectInfo::Type SelType);
+	void Update();
+	void AddFile(FFileDescriptor *InFileDesc);
+	void Clear();
 
 	/**
-	 * Recursively searches starting at Event for stats matching the TotalsFilteringText
-	 * and returns total times
-	 * @param	Event			Start recursively accumulating at this event
-	 * @param	OutTotalTime	Total accumulated inclusive time for all stats matching TotalsFilteringText
+	 * Recursively searches starting at Event for stats matching the TotalsFilteringText.
+	 * and returns total times.
+	 *
+	 * @param Event Start recursively accumulating at this event.
+	 * @param OutTotalTime Total accumulated inclusive time for all stats matching TotalsFilteringText.
 	 */
-	void FindTotalsForPrefilterRecursive(FEventGraphSamplePtr Event, float &OutTotalTime)
-	{
-		FString EventName = Event->_StatName.ToString().ToUpper();
-		if (EventName.Contains(TotalsFilteringText.ToUpper()))
-		{
-			OutTotalTime += Event->_InclusiveTimeMS;
-		}
-		else
-		{
-			for (auto Child : Event->GetChildren())
-			{
-				FindTotalsForPrefilterRecursive(Child, OutTotalTime);
-			}
-		}
-	}
+	void FindTotalsForPrefilterRecursive(TSharedPtr<FEventGraphSample> Event, float &OutTotalTime);
 
 	/* Find total RT and GT times for stats in dump file Desc, matching TotalsFilteringText
  	 * @param	Desc	FileDescriptor for the stats dump file to traverse
 	 */
-	void FindTotalsForPrefilter(TSharedPtr<FFileDescriptor> Desc)
-	{
-		float TotalRenderThreadTime = 0;
-		float TotalGameThreadTime = 0;
+	void FindTotalsForPrefilter(TSharedPtr<FFileDescriptor> Desc);
+	void PrefilterTextCommitted(const FText& InText, const ETextCommit::Type InTextAction);
 
-		// if we don't have a totals filtering text, display the entire RT and GT times
-		bool bGetBaseThreadTimes = false;
-		if (TotalsFilteringText == "")
-		{
-			bGetBaseThreadTimes = true;
-		}
+public:
 
-		FEventGraphSamplePtr RootPtr = FProfilerManager::Get()->GetProfilerSession()->GetEventGraphDataAverage()->GetRoot();
+	//~ SWidget interface
 
-		// start at the root's children, which are the threads
-		for (auto Child : RootPtr->GetChildren())
-		{
-			TArray<FString>EventNameTokens;
-			Child->_ThreadName.ToString().ParseIntoArray(EventNameTokens, TEXT(" "));
+	virtual void Tick(const FGeometry& AllottedGeometry, const double InCurrentTime, const float InDeltaTime) override;
 
-			// get cumulative times for stats on the render thread
-			if (EventNameTokens[0] == "RenderThread")
-			{
-				if (bGetBaseThreadTimes)
-				{
-					TotalsFilteringText = "RenderThread";
-				}
-				FindTotalsForPrefilterRecursive(Child, TotalRenderThreadTime);
-			}
-			// get cumulative times for stats on the game thread
-			else if (EventNameTokens[0] == "GameThread")
-			{
-				if (bGetBaseThreadTimes)
-				{
-					TotalsFilteringText = "GameThread";
-				}
-				FindTotalsForPrefilterRecursive(Child, TotalGameThreadTime);
-			}
+protected:
 
-			if (bGetBaseThreadTimes)
-			{
-				TotalsFilteringText = "";
-			}
-		}
+	TArray<TSharedPtr<FFileDescriptor>> StatsFiles;
 
-		Desc->TimeString = "RT " + FString::Printf(TEXT("%.2f"), TotalRenderThreadTime) + " / GT " + FString::Printf(TEXT("%.2f"), TotalGameThreadTime);
-	}
+	TSharedPtr< SListView<TSharedPtr<FFileDescriptor>> > FileList;
 
-	void PrefilterTextCommitted(const FText& InText, const ETextCommit::Type InTextAction)
-	{
-		TotalsFilteringText = InText.ToString();
-		for (auto Desc : StatsFiles)
-		{
-			Desc->Status = FFileDescriptor::EQueued;
-		}
-	}
+	/** edit box determining for which stats names to show thread time totals in the file list. */
+	TSharedPtr<SEditableTextBox> DisplayTotalsBox;
 
-protected:
-	TArray<TSharedPtr<FFileDescriptor>> StatsFiles;
-	TSharedPtr< SListView<TSharedPtr<FFileDescriptor>> > FileList;
-	TSharedPtr<SEditableTextBox> DisplayTotalsBox;					// edit box determining for which stats names to show thread time totals in the file list
 	FString TotalsFilteringText;
 };
-
-#undef LOCTEXT_NAMESPACE
diff --git a/Engine/Source/Developer/Profiler/Private/Widgets/SProfilerFPSChartPanel.cpp b/Engine/Source/Developer/Profiler/Private/Widgets/SProfilerFPSChartPanel.cpp
--- a/Engine/Source/Developer/Profiler/Private/Widgets/SProfilerFPSChartPanel.cpp
+++ b/Engine/Source/Developer/Profiler/Private/Widgets/SProfilerFPSChartPanel.cpp
@@ -1,10 +1,18 @@
 // Copyright 1998-2016 Epic Games, Inc. All Rights Reserved.
 
 #include "ProfilerPrivatePCH.h"
+#include "ProfilerFPSAnalyzer.h"
+#include "ProfilerManager.h"
+#include "ProfilerSession.h"
+#include "SHistogram.h"
+#include "SProfilerFPSChartPanel.h"
+
 
 #define LOCTEXT_NAMESPACE "SProfilerFPSChartPanel"
 
-class SProfilerFPSStatisticsPanel : public SCompoundWidget
+
+class SProfilerFPSStatisticsPanel
+	: public SCompoundWidget
 {
 public:
 	SLATE_BEGIN_ARGS( SProfilerFPSStatisticsPanel )
@@ -218,12 +226,12 @@ void SProfilerFPSChartPanel::Construct( const FArguments& InArgs )
 }
 END_SLATE_FUNCTION_BUILD_OPTIMIZATION
 
-void SProfilerFPSChartPanel::ProfilerManager_OnViewModeChanged(EProfilerViewMode::Type NewViewMode)
+void SProfilerFPSChartPanel::ProfilerManager_OnViewModeChanged(EProfilerViewMode NewViewMode)
 {
 	auto ProfilerManager = FProfilerManager::Get();
 	if (ProfilerManager.IsValid())
 	{
-		const FProfilerSessionPtr ProfilerSession = ProfilerManager->GetProfilerSession();
+		const TSharedPtr<FProfilerSession> ProfilerSession = ProfilerManager->GetProfilerSession();
 
 		Histogram->SetFPSAnalyzer(ProfilerSession->FPSAnalyzer);
 		StatisticsPanel->SetFPSAnalyzer(ProfilerSession->FPSAnalyzer);
diff --git a/Engine/Source/Developer/Profiler/Private/Widgets/SProfilerFPSChartPanel.h b/Engine/Source/Developer/Profiler/Private/Widgets/SProfilerFPSChartPanel.h
--- a/Engine/Source/Developer/Profiler/Private/Widgets/SProfilerFPSChartPanel.h
+++ b/Engine/Source/Developer/Profiler/Private/Widgets/SProfilerFPSChartPanel.h
@@ -2,10 +2,20 @@
 
 #pragma once
 
-/** A custom widget that acts as a container for widgets like SDataGraph or SEventTree. */
-class SProfilerFPSChartPanel : public SCompoundWidget
+
+enum class EProfilerViewMode;
+class SHistogram;
+class SProfilerFPSStatisticsPanel;
+
+
+/**
+ * A custom widget that acts as a container for widgets like SDataGraph or SEventTree.
+ */
+class SProfilerFPSChartPanel
+	: public SCompoundWidget
 {
 public:
+
 	SLATE_BEGIN_ARGS( SProfilerFPSChartPanel )
 		{}
 
@@ -18,17 +28,20 @@ public:
 	/**
 	 * Construct this widget
 	 *
-	 * @param InArgs - The declaration data for this widget
+	 * @param InArgs The declaration data for this widget
 	 */
 	void Construct( const FArguments& InArgs );
 
 protected:
+
 	/** Called when the status of specified tracked stat has changed. */
-	void ProfilerManager_OnViewModeChanged(EProfilerViewMode::Type NewViewMode);
+	void ProfilerManager_OnViewModeChanged(EProfilerViewMode NewViewMode);
+
+protected:
 
 	/** The descriptions panel of the chart */
-	TSharedPtr<class SHistogram> Histogram;
+	TSharedPtr<SHistogram> Histogram;
 
 	/** The descriptions panel of the chart */
-	TSharedPtr<class SProfilerFPSStatisticsPanel> StatisticsPanel;
+	TSharedPtr<SProfilerFPSStatisticsPanel> StatisticsPanel;
 };
\ No newline at end of file
diff --git a/Engine/Source/Developer/Profiler/Private/Widgets/SProfilerGraphPanel.cpp b/Engine/Source/Developer/Profiler/Private/Widgets/SProfilerGraphPanel.cpp
--- a/Engine/Source/Developer/Profiler/Private/Widgets/SProfilerGraphPanel.cpp
+++ b/Engine/Source/Developer/Profiler/Private/Widgets/SProfilerGraphPanel.cpp
@@ -1,184 +1,197 @@
 // Copyright 1998-2016 Epic Games, Inc. All Rights Reserved.
 
 #include "ProfilerPrivatePCH.h"
+#include "ProfilerManager.h"
+#include "SDataGraph.h"
+#include "SProfilerGraphPanel.h"
+#include "SProfilerMiniView.h"
+#include "SProfilerThreadView.h"
+
 
 #define LOCTEXT_NAMESPACE "SProfilerGraphPanel"
 
+
 SProfilerGraphPanel::SProfilerGraphPanel()
-	: bLockMiniViewState( false )
-	, NumDataPoints( 0 )
-	, NumVisiblePoints( 0 )
-	, GraphOffset( 0 )
-	, ViewMode( EProfilerViewMode::InvalidOrMax )
-{}
+	: bLockMiniViewState(false)
+	, NumDataPoints(0)
+	, NumVisiblePoints(0)
+	, GraphOffset(0)
+	, ViewMode(EProfilerViewMode::InvalidOrMax)
+{ }
+
 
 SProfilerGraphPanel::~SProfilerGraphPanel()
 {
 	// Remove ourselves from the profiler manager.
 	auto ProfilerManager = FProfilerManager::Get();
-	if( ProfilerManager.IsValid() )
+	if(ProfilerManager.IsValid())
 	{
-		ProfilerManager->OnTrackedStatChanged().RemoveAll( this );
-		ProfilerManager->OnViewModeChanged().RemoveAll( this );
+		ProfilerManager->OnTrackedStatChanged().RemoveAll(this);
+		ProfilerManager->OnViewModeChanged().RemoveAll(this);
 
-		DataGraph->OnSelectionChangedForIndex().RemoveAll( ProfilerManager.Get() );
+		DataGraph->OnSelectionChangedForIndex().RemoveAll(ProfilerManager.Get());
 
-		if( ProfilerMiniView.Pin().IsValid() )
+		if(ProfilerMiniView.Pin().IsValid())
 		{
-			ProfilerMiniView.Pin()->OnSelectionBoxChanged().RemoveAll( this );
+			ProfilerMiniView.Pin()->OnSelectionBoxChanged().RemoveAll(this);
 		}
 	}
 
-	ThreadView->OnViewPositionXChanged().RemoveAll( this );
-	ThreadView->OnViewPositionYChanged().RemoveAll( this );
+	ThreadView->OnViewPositionXChanged().RemoveAll(this);
+	ThreadView->OnViewPositionYChanged().RemoveAll(this);
 }
 
-void SProfilerGraphPanel::Construct( const FArguments& InArgs )
+
+void SProfilerGraphPanel::Construct(const FArguments& InArgs)
 {
 	ChildSlot
 	[
-		SNew( SBorder )
-		.BorderImage( FEditorStyle::GetBrush( "ToolPanel.GroupBorder" ) )
-		.Padding( 2.0f )
+		SNew(SBorder)
+		.BorderImage(FEditorStyle::GetBrush("ToolPanel.GroupBorder"))
+		.Padding(2.0f)
 		[
-			SNew( SHorizontalBox )
+			SNew(SHorizontalBox)
 
 			+ SHorizontalBox::Slot()
-			.FillWidth( 1.0f )
+			.FillWidth(1.0f)
 			[
-				SNew( SVerticalBox )
+				SNew(SVerticalBox)
 
 				// At this moment only one widget of these two can be visible at once.
 				//
 				// DataGraph
 				+SVerticalBox::Slot()
-				.FillHeight( 1.0f )
+				.FillHeight(1.0f)
 				[
-					SAssignNew( DataGraph, SDataGraph )
-					.OnGraphOffsetChanged( this, &SProfilerGraphPanel::OnDataGraphGraphOffsetChanged )
+					SAssignNew(DataGraph, SDataGraph)
+					.OnGraphOffsetChanged(this, &SProfilerGraphPanel::OnDataGraphGraphOffsetChanged)
 				]
 
 				// ThreadView
 				+ SVerticalBox::Slot()
-				.FillHeight( 1.0f )
+				.FillHeight(1.0f)
 				[
-					SAssignNew( ThreadView, SProfilerThreadView )
+					SAssignNew(ThreadView, SProfilerThreadView)
 				]
 
 				+ SVerticalBox::Slot()
 				.AutoHeight()
 				[
-					SAssignNew( HorizontalScrollBar, SScrollBar )
-					.Orientation( Orient_Horizontal )
-					.AlwaysShowScrollbar( true )
-					.Visibility( EVisibility::Visible )
-					.Thickness( FVector2D( 8.0f, 8.0f ) )
-					.OnUserScrolled( this, &SProfilerGraphPanel::HorizontalScrollBar_OnUserScrolled )
+					SAssignNew(HorizontalScrollBar, SScrollBar)
+					.Orientation(Orient_Horizontal)
+					.AlwaysShowScrollbar(true)
+					.Visibility(EVisibility::Visible)
+					.Thickness(FVector2D(8.0f, 8.0f))
+					.OnUserScrolled(this, &SProfilerGraphPanel::HorizontalScrollBar_OnUserScrolled)
 				]
 			]
 
 			+ SHorizontalBox::Slot()
 			.AutoWidth()
 			[
-				SAssignNew( VerticalScrollBar, SScrollBar )
-				.Orientation( Orient_Vertical )
-				.AlwaysShowScrollbar( true )
-				.Visibility( EVisibility::Visible )
-				.Thickness( FVector2D( 8.0f, 8.0f ) )
-				.OnUserScrolled( this, &SProfilerGraphPanel::VerticalScrollBar_OnUserScrolled )
+				SAssignNew(VerticalScrollBar, SScrollBar)
+				.Orientation(Orient_Vertical)
+				.AlwaysShowScrollbar(true)
+				.Visibility(EVisibility::Visible)
+				.Thickness(FVector2D(8.0f, 8.0f))
+				.OnUserScrolled(this, &SProfilerGraphPanel::VerticalScrollBar_OnUserScrolled)
 			]
 		]
 	];
 
-	HorizontalScrollBar->SetState( 0.0f, 1.0f );
-	VerticalScrollBar->SetState( 0.0f, 1.0f );
+	HorizontalScrollBar->SetState(0.0f, 1.0f);
+	VerticalScrollBar->SetState(0.0f, 1.0f);
 
 	// Register ourselves with the profiler manager.
-	FProfilerManager::Get()->OnTrackedStatChanged().AddSP( this, &SProfilerGraphPanel::ProfilerManager_OnTrackedStatChanged );
-	FProfilerManager::Get()->OnViewModeChanged().AddSP( this, &SProfilerGraphPanel::ProfilerManager_OnViewModeChanged );
+	FProfilerManager::Get()->OnTrackedStatChanged().AddSP(this, &SProfilerGraphPanel::ProfilerManager_OnTrackedStatChanged);
+	FProfilerManager::Get()->OnViewModeChanged().AddSP(this, &SProfilerGraphPanel::ProfilerManager_OnViewModeChanged);
 
-	DataGraph->OnSelectionChangedForIndex().AddSP( FProfilerManager::Get().ToSharedRef(), &FProfilerManager::DataGraph_OnSelectionChangedForIndex );
+	DataGraph->OnSelectionChangedForIndex().AddSP(FProfilerManager::Get().ToSharedRef(), &FProfilerManager::DataGraph_OnSelectionChangedForIndex);
 
-	ThreadView->OnViewPositionXChanged().AddSP( this, &SProfilerGraphPanel::ThreadView_OnViewPositionXChanged );
-	ThreadView->OnViewPositionYChanged().AddSP( this, &SProfilerGraphPanel::ThreadView_OnViewPositionYChanged );
+	ThreadView->OnViewPositionXChanged().AddSP(this, &SProfilerGraphPanel::ThreadView_OnViewPositionXChanged);
+	ThreadView->OnViewPositionYChanged().AddSP(this, &SProfilerGraphPanel::ThreadView_OnViewPositionYChanged);
 }
 
 
-void SProfilerGraphPanel::Tick( const FGeometry& AllottedGeometry, const double InCurrentTime, const float InDeltaTime )
+void SProfilerGraphPanel::Tick(const FGeometry& AllottedGeometry, const double InCurrentTime, const float InDeltaTime)
 {
 	UpdateInternals();
 }
 
-void SProfilerGraphPanel::HorizontalScrollBar_OnUserScrolled( float ScrollOffset )
+
+void SProfilerGraphPanel::HorizontalScrollBar_OnUserScrolled(float ScrollOffset)
 {
-	if( ViewMode == EProfilerViewMode::LineIndexBased )
+	if(ViewMode == EProfilerViewMode::LineIndexBased)
 	{
-		const float ThumbSizeFraction = FMath::Min( (float)NumVisiblePoints / (float)NumDataPoints, 1.0f );
-		ScrollOffset = FMath::Min( ScrollOffset, 1.0f - ThumbSizeFraction );
+		const float ThumbSizeFraction = FMath::Min((float)NumVisiblePoints / (float)NumDataPoints, 1.0f);
+		ScrollOffset = FMath::Min(ScrollOffset, 1.0f - ThumbSizeFraction);
 
-		HorizontalScrollBar->SetState( ScrollOffset, ThumbSizeFraction );
+		HorizontalScrollBar->SetState(ScrollOffset, ThumbSizeFraction);
 
-		GraphOffset = FMath::TruncToInt( ScrollOffset * NumDataPoints );
-		DataGraph->ScrollTo( GraphOffset );
+		GraphOffset = FMath::TruncToInt(ScrollOffset * NumDataPoints);
+		DataGraph->ScrollTo(GraphOffset);
 
-		ProfilerMiniView.Pin()->MoveWithoutZoomSelectionBox( GraphOffset, GraphOffset + NumVisiblePoints );
+		ProfilerMiniView.Pin()->MoveWithoutZoomSelectionBox(GraphOffset, GraphOffset + NumVisiblePoints);
 	}
-	else if( ViewMode == EProfilerViewMode::ThreadViewTimeBased )
+	else if(ViewMode == EProfilerViewMode::ThreadViewTimeBased)
 	{
-		ThreadView->SetPositionXToByScrollBar( ScrollOffset );
+		ThreadView->SetPositionXToByScrollBar(ScrollOffset);
 	}
 }
 
 
-void SProfilerGraphPanel::VerticalScrollBar_OnUserScrolled( float ScrollOffset )
+void SProfilerGraphPanel::VerticalScrollBar_OnUserScrolled(float ScrollOffset)
 {
-	if( ViewMode == EProfilerViewMode::LineIndexBased )
+	if(ViewMode == EProfilerViewMode::LineIndexBased)
 	{
 
 	}
-	else if( ViewMode == EProfilerViewMode::ThreadViewTimeBased )
+	else if(ViewMode == EProfilerViewMode::ThreadViewTimeBased)
 	{
-		ThreadView->SetPositonYTo( ScrollOffset );
+		ThreadView->SetPositonYTo(ScrollOffset);
 	}
 }
 
 
-void SProfilerGraphPanel::ProfilerManager_OnTrackedStatChanged( const FTrackedStatPtr& TrackedStat, bool bIsTracked )
+void SProfilerGraphPanel::ProfilerManager_OnTrackedStatChanged(const TSharedPtr<FTrackedStat>& TrackedStat, bool bIsTracked)
 {
-	if( bIsTracked )
+	if(bIsTracked)
 	{
-		DataGraph->AddInnerGraph( TrackedStat );
+		DataGraph->AddInnerGraph(TrackedStat);
 	}
 	else
 	{
-		DataGraph->RemoveInnerGraph( TrackedStat->StatID );
+		DataGraph->RemoveInnerGraph(TrackedStat->StatID);
 	}
 }
 
-void SProfilerGraphPanel::OnDataGraphGraphOffsetChanged( int32 InGraphOffset )
+
+void SProfilerGraphPanel::OnDataGraphGraphOffsetChanged(int32 InGraphOffset)
 {
 	GraphOffset = InGraphOffset;
 	SetScrollBarState();
-	ProfilerMiniView.Pin()->MoveWithoutZoomSelectionBox( InGraphOffset, InGraphOffset + NumVisiblePoints );
+	ProfilerMiniView.Pin()->MoveWithoutZoomSelectionBox(InGraphOffset, InGraphOffset + NumVisiblePoints);
 }
 
-void SProfilerGraphPanel::MiniView_OnSelectionBoxChanged( int32 FrameStart, int32 FrameEnd )
+
+void SProfilerGraphPanel::MiniView_OnSelectionBoxChanged(int32 FrameStart, int32 FrameEnd)
 {
-	if( ViewMode == EProfilerViewMode::LineIndexBased )
+	if(ViewMode == EProfilerViewMode::LineIndexBased)
 	{
 		GraphOffset = FrameStart;
 		SetScrollBarState();
-		DataGraph->ScrollTo( GraphOffset );
+		DataGraph->ScrollTo(GraphOffset);
 	}
-	else if( ViewMode == EProfilerViewMode::ThreadViewTimeBased )
+	else if(ViewMode == EProfilerViewMode::ThreadViewTimeBased)
 	{
-		TGuardValue<bool> LockedMiniViewState( bLockMiniViewState, true );
+		TGuardValue<bool> LockedMiniViewState(bLockMiniViewState, true);
 		// Update thread-view state.
-		ThreadView->SetFrameRange( FrameStart, FrameEnd );
+		ThreadView->SetFrameRange(FrameStart, FrameEnd);
 	}
 }
 
-void SProfilerGraphPanel::ThreadView_OnViewPositionXChanged( double FrameStartMS, double FrameEndMS, double MaxEndTimeMS, int32 FrameStart, int32 FrameEnd )
+
+void SProfilerGraphPanel::ThreadView_OnViewPositionXChanged(double FrameStartMS, double FrameEndMS, double MaxEndTimeMS, int32 FrameStart, int32 FrameEnd)
 {
 	const double MinStartTimeMS = 0.0f;
 	const double TotalRangeMS = MaxEndTimeMS - MinStartTimeMS;
@@ -188,16 +201,17 @@ void SProfilerGraphPanel::ThreadView_OnViewPositionXChanged( double FrameStartMS
 	const double OffsetFraction = FrameStartMS / MaxEndTimeMS;
 
 	// Update horizontal scroll bar state.
-	HorizontalScrollBar->SetState( (float)OffsetFraction, (float)ThumbSizeFraction );
+	HorizontalScrollBar->SetState((float)OffsetFraction, (float)ThumbSizeFraction);
 
-	if( !bLockMiniViewState )
+	if(!bLockMiniViewState)
 	{
 		// Update profiler mini-view state.
-		ProfilerMiniView.Pin()->MoveAndZoomSelectionBox( FrameStart, FrameEnd );
+		ProfilerMiniView.Pin()->MoveAndZoomSelectionBox(FrameStart, FrameEnd);
 	}
 }
 
-void SProfilerGraphPanel::ThreadView_OnViewPositionYChanged( double PosStartY, double PosYEnd, double MaxPosY )
+
+void SProfilerGraphPanel::ThreadView_OnViewPositionYChanged(double PosStartY, double PosYEnd, double MaxPosY)
 {
 	const double MinPosY = 0.0f;
 	const double TotalRangeY = MaxPosY - MinPosY;
@@ -206,30 +220,31 @@ void SProfilerGraphPanel::ThreadView_OnViewPositionYChanged( double PosStartY, d
 	const double ThumbSizeFraction = ThisRangeY / TotalRangeY;
 	const double OffsetFraction = PosStartY / MaxPosY;
 
-	VerticalScrollBar->SetState( (float)OffsetFraction, (float)ThumbSizeFraction );
+	VerticalScrollBar->SetState((float)OffsetFraction, (float)ThumbSizeFraction);
 }
 
-void SProfilerGraphPanel::ProfilerManager_OnViewModeChanged( EProfilerViewMode::Type NewViewMode )
+
+void SProfilerGraphPanel::ProfilerManager_OnViewModeChanged(EProfilerViewMode NewViewMode)
 {
-	if( NewViewMode == EProfilerViewMode::LineIndexBased )
+	if(NewViewMode == EProfilerViewMode::LineIndexBased)
 	{
-		VerticalScrollBar->SetVisibility( EVisibility::Collapsed );
-		VerticalScrollBar->SetEnabled( false );
+		VerticalScrollBar->SetVisibility(EVisibility::Collapsed);
+		VerticalScrollBar->SetEnabled(false);
 
-		DataGraph->SetVisibility( EVisibility::Visible );
-		DataGraph->SetEnabled( true );
-		ThreadView->SetVisibility( EVisibility::Collapsed );
-		ThreadView->SetEnabled( false );
+		DataGraph->SetVisibility(EVisibility::Visible);
+		DataGraph->SetEnabled(true);
+		ThreadView->SetVisibility(EVisibility::Collapsed);
+		ThreadView->SetEnabled(false);
 	}
-	else if( NewViewMode == EProfilerViewMode::ThreadViewTimeBased )
+	else if(NewViewMode == EProfilerViewMode::ThreadViewTimeBased)
 	{
-		VerticalScrollBar->SetVisibility( EVisibility::Visible );
-		VerticalScrollBar->SetEnabled( true );
+		VerticalScrollBar->SetVisibility(EVisibility::Visible);
+		VerticalScrollBar->SetEnabled(true);
 
-		DataGraph->SetVisibility( EVisibility::Collapsed );
-		DataGraph->SetEnabled( false );
-		ThreadView->SetVisibility( EVisibility::Visible );
-		ThreadView->SetEnabled( true );
+		DataGraph->SetVisibility(EVisibility::Collapsed);
+		DataGraph->SetEnabled(false);
+		ThreadView->SetVisibility(EVisibility::Visible);
+		ThreadView->SetEnabled(true);
 	}
 
 	ViewMode = NewViewMode;
@@ -238,34 +253,34 @@ void SProfilerGraphPanel::ProfilerManager_OnViewModeChanged( EProfilerViewMode::
 
 void SProfilerGraphPanel::UpdateInternals()
 {
-	if( ViewMode == EProfilerViewMode::LineIndexBased )
+	if(ViewMode == EProfilerViewMode::LineIndexBased)
 	{
 		NumVisiblePoints = DataGraph->GetNumVisiblePoints();
 		NumDataPoints = DataGraph->GetNumDataPoints();
 
 		SetScrollBarState();
-		ProfilerMiniView.Pin()->MoveWithoutZoomSelectionBox( GraphOffset, GraphOffset + NumVisiblePoints );
+		ProfilerMiniView.Pin()->MoveWithoutZoomSelectionBox(GraphOffset, GraphOffset + NumVisiblePoints);
 
-		if( FProfilerManager::Get()->IsLivePreview() )
+		if(FProfilerManager::Get()->IsLivePreview())
 		{
 			// Scroll to the end.
-			HorizontalScrollBar_OnUserScrolled( 1.0f );
+			HorizontalScrollBar_OnUserScrolled(1.0f);
 		}
 	}
-	else if( ViewMode == EProfilerViewMode::ThreadViewTimeBased )
+	else if(ViewMode == EProfilerViewMode::ThreadViewTimeBased)
 	{
 	}
 }
 
+
 void SProfilerGraphPanel::SetScrollBarState()
 {
 	// Note that the maximum offset is 1.0-ThumbSizeFraction.
 	// If the user can view 1/3 of the items in a single page, the maximum offset will be ~0.667f
-	const float ThumbSizeFraction = FMath::Min( (float)NumVisiblePoints / (float)NumDataPoints, 1.0f );
+	const float ThumbSizeFraction = FMath::Min((float)NumVisiblePoints / (float)NumDataPoints, 1.0f);
 	const float OffsetFraction = (float)GraphOffset / (float)NumDataPoints;
-	HorizontalScrollBar->SetState( OffsetFraction, ThumbSizeFraction );
+	HorizontalScrollBar->SetState(OffsetFraction, ThumbSizeFraction);
 }
 
 
-
 #undef LOCTEXT_NAMESPACE
diff --git a/Engine/Source/Developer/Profiler/Private/Widgets/SProfilerGraphPanel.h b/Engine/Source/Developer/Profiler/Private/Widgets/SProfilerGraphPanel.h
--- a/Engine/Source/Developer/Profiler/Private/Widgets/SProfilerGraphPanel.h
+++ b/Engine/Source/Developer/Profiler/Private/Widgets/SProfilerGraphPanel.h
@@ -2,69 +2,85 @@
 
 #pragma once
 
+
+enum class EProfilerViewMode;
+class FTrackedStat;
+class SDataGraph;
+class SProfilerMiniView;
+class SProfilerThreadView;
+class SScrollBar;
+
+
 /** A custom widget that acts as a container for widgets like SDataGraph or SEventTree. */
-class SProfilerGraphPanel : public SCompoundWidget
+class SProfilerGraphPanel
+	: public SCompoundWidget
 {
 public:
+
 	/** Default constructor. */
 	SProfilerGraphPanel();
 		
 	/** Virtual destructor. */
 	virtual ~SProfilerGraphPanel();
 
-	SLATE_BEGIN_ARGS( SProfilerGraphPanel ){}
+	SLATE_BEGIN_ARGS(SProfilerGraphPanel){}
 	SLATE_END_ARGS()
 
 	/**
 	 * Construct this widget
 	 *
-	 * @param InArgs - The declaration data for this widget
+	 * @param InArgs The declaration data for this widget
 	 */
-	void Construct( const FArguments& InArgs );
-
-	// SWidget interface
-	virtual void Tick( const FGeometry& AllottedGeometry, const double InCurrentTime, const float InDeltaTime ) override;
-	// End of SWidget interface
+	void Construct(const FArguments& InArgs);
 
 	TSharedPtr<SDataGraph>& GetMainDataGraph()
 	{
 		return DataGraph;
 	}
 
+public:
+
+	//~ SWidget interface
+
+	virtual void Tick(const FGeometry& AllottedGeometry, const double InCurrentTime, const float InDeltaTime) override;
+
 protected:
+
 	/** Called when the status of specified tracked stat has changed. */
-	void ProfilerManager_OnTrackedStatChanged( const FTrackedStatPtr& TrackedStat, bool bIsTracked );
+	void ProfilerManager_OnTrackedStatChanged(const TSharedPtr<FTrackedStat>& TrackedStat, bool bIsTracked);
 
 	/**
 	 * Called when the user scrolls the scrollbar
 	 *
-	 * @param ScrollOffset - scroll offset as a fraction between 0 and 1.
+	 * @param ScrollOffset Scroll offset as a fraction between 0 and 1.
 	 */
-	void HorizontalScrollBar_OnUserScrolled( float ScrollOffset );
+	void HorizontalScrollBar_OnUserScrolled(float ScrollOffset);
 
-	void VerticalScrollBar_OnUserScrolled( float ScrollOffset );
+	void VerticalScrollBar_OnUserScrolled(float ScrollOffset);
 
 	/** Called when the frame offset has been changed in the data graph widget. */
-	void OnDataGraphGraphOffsetChanged( int32 InFrameOffset );
+	void OnDataGraphGraphOffsetChanged(int32 InFrameOffset);
 
-	void ProfilerManager_OnViewModeChanged( EProfilerViewMode::Type NewViewMode );
+	void ProfilerManager_OnViewModeChanged(EProfilerViewMode NewViewMode);
 
 public:
+
 	/** Called when the selection box has been changed in the profiler mini-view widget. */
-	void MiniView_OnSelectionBoxChanged( int32 FrameStart, int32 FrameEnd );
+	void MiniView_OnSelectionBoxChanged(int32 FrameStart, int32 FrameEnd);
 
-	void ThreadView_OnViewPositionXChanged( double FrameStartMS, double FrameEndMS, double MaxEndTimeMS, int32 FrameStart, int32 FrameEnd );
-	void ThreadView_OnViewPositionYChanged( double PosYStart, double PosYEnd, double MaxPosY );
+	void ThreadView_OnViewPositionXChanged(double FrameStartMS, double FrameEndMS, double MaxEndTimeMS, int32 FrameStart, int32 FrameEnd);
+	void ThreadView_OnViewPositionYChanged(double PosYStart, double PosYEnd, double MaxPosY);
 
 protected:
+
 	/** Sets state of the scroll bar. */
 	void SetScrollBarState();
 
-
 	void UpdateInternals();
 
 //protected:
 public:
+
 	/** Holds the data graph widget. */
 	TSharedPtr<SDataGraph> DataGraph;
 
@@ -92,5 +108,6 @@ public:
 	/** Current offset of the graph, index of the first visible graph point. */
 	int32 GraphOffset;
 
-	EProfilerViewMode::Type ViewMode;
+	/** Current view mode. */
+	EProfilerViewMode ViewMode;
 };
\ No newline at end of file
diff --git a/Engine/Source/Developer/Profiler/Private/Widgets/SProfilerMiniView.cpp b/Engine/Source/Developer/Profiler/Private/Widgets/SProfilerMiniView.cpp
--- a/Engine/Source/Developer/Profiler/Private/Widgets/SProfilerMiniView.cpp
+++ b/Engine/Source/Developer/Profiler/Private/Widgets/SProfilerMiniView.cpp
@@ -1,11 +1,9 @@
 // Copyright 1998-2016 Epic Games, Inc. All Rights Reserved.
 
 #include "ProfilerPrivatePCH.h"
+#include "ProfilerSession.h"
+#include "SProfilerMiniView.h"
 
-#define LOCTEXT_NAMESPACE "SProfilerMiniView"
-
-// SProfilerMiniView
-#undef LOCTEXT_NAMESPACE
 
 SProfilerMiniView::SProfilerMiniView()
 	: bIsActiveTimerRegistered( false )
@@ -579,7 +577,7 @@ void SProfilerMiniView::BindCommands()
 
 
 
-void SProfilerMiniView::AddThreadTime( int32 InFrameIndex, const TMap<uint32, float>& InThreadMS, const FProfilerStatMetaDataRef& InStatMetaData )
+void SProfilerMiniView::AddThreadTime( int32 InFrameIndex, const TMap<uint32, float>& InThreadMS, const TSharedRef<FProfilerStatMetaData>& InStatMetaData )
 {
 	FFrameThreadTimes FrameThreadTimes;
 	FrameThreadTimes.FrameNumber = InFrameIndex;
diff --git a/Engine/Source/Developer/Profiler/Private/Widgets/SProfilerMiniView.h b/Engine/Source/Developer/Profiler/Private/Widgets/SProfilerMiniView.h
--- a/Engine/Source/Developer/Profiler/Private/Widgets/SProfilerMiniView.h
+++ b/Engine/Source/Developer/Profiler/Private/Widgets/SProfilerMiniView.h
@@ -2,9 +2,9 @@
 
 #pragma once
 
-/*-----------------------------------------------------------------------------
-	Basic structures
------------------------------------------------------------------------------*/
+
+class FProfilerStatMetaData;
+
 
 struct FFrameThreadTimes
 {
@@ -230,7 +230,7 @@ public:
 	void BindCommands();
 
 	/** For the specified frame ands a thread data. */
-	void AddThreadTime( int32 InFrameIndex, const TMap<uint32, float>& InThreadMS, const FProfilerStatMetaDataRef& InStatMetaData );
+	void AddThreadTime( int32 InFrameIndex, const TMap<uint32, float>& InThreadMS, const TSharedRef<FProfilerStatMetaData>& InStatMetaData );
 
 	/**
 	 *	Moves the selection box to the location specified by FrameStart and FrameEnd. 
@@ -349,7 +349,7 @@ protected:
 	float MaxFrameTime;
 
 	/** Shared pointer to the stats' metadata. */
-	FProfilerStatMetaDataPtr StatMetadata;
+	TSharedPtr<FProfilerStatMetaData> StatMetadata;
 
 	/*-----------------------------------------------------------------------------
 		UI variables
diff --git a/Engine/Source/Developer/Profiler/Private/Widgets/SProfilerSettings.cpp b/Engine/Source/Developer/Profiler/Private/Widgets/SProfilerSettings.cpp
--- a/Engine/Source/Developer/Profiler/Private/Widgets/SProfilerSettings.cpp
+++ b/Engine/Source/Developer/Profiler/Private/Widgets/SProfilerSettings.cpp
@@ -1,10 +1,13 @@
 // Copyright 1998-2016 Epic Games, Inc. All Rights Reserved.
 
 #include "ProfilerPrivatePCH.h"
+#include "ProfilerManager.h"
 #include "SProfilerSettings.h"
 
+
 #define LOCTEXT_NAMESPACE "SProfilerSettings"
 
+
 BEGIN_SLATE_FUNCTION_BUILD_OPTIMIZATION
 void SProfilerSettings::Construct( const FArguments& InArgs )
 {
diff --git a/Engine/Source/Developer/Profiler/Private/Widgets/SProfilerSettings.h b/Engine/Source/Developer/Profiler/Private/Widgets/SProfilerSettings.h
--- a/Engine/Source/Developer/Profiler/Private/Widgets/SProfilerSettings.h
+++ b/Engine/Source/Developer/Profiler/Private/Widgets/SProfilerSettings.h
@@ -3,10 +3,17 @@
 #pragma once
 
 
-/** Widget used to modify settings for the profiler, created on demand and destroyed on close. */
-class SProfilerSettings : public SCompoundWidget
+class FProfilerSettings;
+
+
+/**
+ * Widget used to modify settings for the profiler, created on demand and destroyed on close.
+ */
+class SProfilerSettings
+	: public SCompoundWidget
 {
 public:
+
 	SLATE_BEGIN_ARGS( SProfilerSettings ){}
 		SLATE_EVENT( FSimpleDelegate, OnClose )
 		SLATE_ARGUMENT( FProfilerSettings*, SettingPtr )
@@ -15,19 +22,16 @@ public:
 	/**
 	 * Construct this widget
 	 *
-	 * @param	InArgs	The declaration data for this widget
+	 * @param InArgs The declaration data for this widget.
 	 */
 	void Construct( const FArguments& InArgs );
 
 private:
-	void AddTitle( const FText& TitleText, const TSharedRef<SGridPanel>& Grid, int32& RowPos );
 
+	void AddTitle( const FText& TitleText, const TSharedRef<SGridPanel>& Grid, int32& RowPos );
 	void AddSeparator( const TSharedRef<SGridPanel>& Grid, int32& RowPos );
-
 	void AddHeader( const FText& HeaderText, const TSharedRef<SGridPanel>& Grid, int32& RowPos);
-
 	void AddOption( const FText& OptionName, const FText& OptionDesc, bool& Value, const bool& Default, const TSharedRef<SGridPanel>& Grid, int32& RowPos );
-
 	void AddFooter( const TSharedRef<SGridPanel>& Grid, int32& RowPos );
 
 	FReply SaveAndClose_OnClicked();
@@ -40,9 +44,10 @@ private:
 	ECheckBoxState OptionValue_IsChecked( const bool* ValuePtr ) const;
 
 private:
+
 	/** Delegate to call when this profiler settings is closed. */
 	FSimpleDelegate OnClose;
 
 	/** Pointer to the profiler settings. */
 	FProfilerSettings* SettingPtr;
-};
\ No newline at end of file
+};
diff --git a/Engine/Source/Developer/Profiler/Private/Widgets/SProfilerThreadView.cpp b/Engine/Source/Developer/Profiler/Private/Widgets/SProfilerThreadView.cpp
--- a/Engine/Source/Developer/Profiler/Private/Widgets/SProfilerThreadView.cpp
+++ b/Engine/Source/Developer/Profiler/Private/Widgets/SProfilerThreadView.cpp
@@ -1,11 +1,9 @@
 // Copyright 1998-2016 Epic Games, Inc. All Rights Reserved.
 
 #include "ProfilerPrivatePCH.h"
+#include "ProfilerStream.h"
+#include "SProfilerThreadView.h"
 
-#define LOCTEXT_NAMESPACE "SProfilerThreadView"
-
-// SProfilerThreadView
-#undef LOCTEXT_NAMESPACE
 
 SProfilerThreadView::SProfilerThreadView()
 {
@@ -655,9 +653,98 @@ void SProfilerThreadView::BindCommands()
 
 }
 
+
+/* SProfilerThreadView interface
+ *****************************************************************************/
+
+void SProfilerThreadView::SetPositionXToByScrollBar( double ScrollOffset ) 
+{
+	SetPositionX( ScrollOffset*TotalRangeXMS );
+}
+
+
+void SProfilerThreadView::SetPositionX( double NewPositionXMS )
+{
+	const double ClampedPositionXMS = FMath::Clamp( NewPositionXMS, 0.0, TotalRangeXMS - RangeXMS );
+	SetTimeRange( ClampedPositionXMS, ClampedPositionXMS + RangeXMS, true );
+}
+
+
+void SProfilerThreadView::SetPositonYTo( double ScrollOffset )
+{
+
+}
+
+
+void SProfilerThreadView::SetTimeRange( double StartTimeMS, double EndTimeMS, bool bBroadcast )
+{
+	check( EndTimeMS > StartTimeMS );
+
+	PositionXMS = StartTimeMS;
+	RangeXMS = EndTimeMS - StartTimeMS;
+	FramesIndices = ProfilerStream->GetFramesIndicesForTimeRange( StartTimeMS, EndTimeMS );
+
+	bUpdateData = true;
+		
+	//UE_LOG( LogTemp, Log, TEXT( "StartTimeMS=%f, EndTimeMS=%f, bBroadcast=%1i FramesIndices=%3i,%3i" ), StartTimeMS, EndTimeMS, (int)bBroadcast, FramesIndices.X, FramesIndices.Y );
+	if( bBroadcast )
+	{
+		ViewPositionXChangedEvent.Broadcast( StartTimeMS, EndTimeMS, TotalRangeXMS, FramesIndices.X, FramesIndices.Y );
+	}
+}
+
+
+void SProfilerThreadView::SetFrameRange(int32 FrameStart, int32 FrameEnd)
+{
+	const double EndTimeMS = ProfilerStream->GetElapsedFrameTimeMS(FrameEnd);
+	const double StartTimeMS = ProfilerStream->GetElapsedFrameTimeMS(FrameStart) - ProfilerStream->GetFrameTimeMS(FrameStart);
+	SetTimeRange(StartTimeMS, EndTimeMS, true);
+}
+
+
+void SProfilerThreadView::AttachProfilerStream(const FProfilerStream& InProfilerStream)
+{
+	ProfilerStream = &InProfilerStream;
+
+	TotalRangeXMS = ProfilerStream->GetElapsedTime();
+	TotalRangeY = ProfilerStream->GetNumThreads()*FProfilerUIStream::DEFAULT_VISIBLE_THREAD_DEPTH;
+
+	// Display the first frame.
+	const FProfilerFrame* ProfilerFrame = ProfilerStream->GetProfilerFrame(0);
+	SetTimeRange(ProfilerFrame->Root->CycleCounterStartTimeMS, ProfilerFrame->Root->CycleCounterEndTimeMS);
+}
+
+
+/* SProfilerThreadView implementation
+ *****************************************************************************/
+
 void SProfilerThreadView::ProcessData()
 {
 	//	SCOPE_LOG_TIME_FUNC();
 
 	ProfilerUIStream.GenerateUIStream( *ProfilerStream, PositionXMS, PositionXMS + RangeXMS, ZoomFactorX, NumMillisecondsPerWindow, NumPixelsPerMillisecond, NumMillisecondsPerSample );
 }
+
+
+bool SProfilerThreadView::IsReady() const
+{
+	return ProfilerStream && ProfilerStream->GetNumFrames() > 0;
+}
+
+
+bool SProfilerThreadView::ShouldUpdateData()
+{
+	return bUpdateData;
+}
+
+
+void SProfilerThreadView::UpdateInternalConstants()
+{
+	ZoomFactorX = (double)NUM_MILLISECONDS_PER_WINDOW / RangeXMS;
+	RangeY = FMath::RoundToFloat(ThisGeometry.Size.Y / (double)NUM_PIXELS_PER_ROW);
+
+	const double Aspect = ThisGeometry.Size.X / NUM_MILLISECONDS_PER_WINDOW * ZoomFactorX;
+	NumMillisecondsPerWindow = (double)ThisGeometry.Size.X / Aspect;
+	NumPixelsPerMillisecond = (double)ThisGeometry.Size.X / NumMillisecondsPerWindow;
+	NumMillisecondsPerSample = NumMillisecondsPerWindow / (double)ThisGeometry.Size.X * (double)MIN_NUM_PIXELS_PER_SAMPLE;
+}
diff --git a/Engine/Source/Developer/Profiler/Private/Widgets/SProfilerThreadView.h b/Engine/Source/Developer/Profiler/Private/Widgets/SProfilerThreadView.h
--- a/Engine/Source/Developer/Profiler/Private/Widgets/SProfilerThreadView.h
+++ b/Engine/Source/Developer/Profiler/Private/Widgets/SProfilerThreadView.h
@@ -2,17 +2,14 @@
 
 #pragma once
 
-/*-----------------------------------------------------------------------------
-	Basic structures
------------------------------------------------------------------------------*/
+#include "ProfilerStream.h"
 
 
-/*-----------------------------------------------------------------------------
-	Declarations
------------------------------------------------------------------------------*/
-
-/** Widget used to present thread data in the mini-view. */
-class SProfilerThreadView : public SCompoundWidget
+/**
+ * Widget used to present thread data in the mini-view.
+ */
+class SProfilerThreadView
+	: public SCompoundWidget
 {
 	enum 
 	{
@@ -163,67 +160,26 @@ public:
 	 *	Changes the position-x of the thread view.
 	 *	Called by the horizontal scroll bar.
 	 */
-	void SetPositionXToByScrollBar( double ScrollOffset ) 
-	{
-		SetPositionX( ScrollOffset*TotalRangeXMS );
-	}
-
-	void SetPositionX( double NewPositionXMS )
-	{
-		const double ClampedPositionXMS = FMath::Clamp( NewPositionXMS, 0.0, TotalRangeXMS - RangeXMS );
-		SetTimeRange( ClampedPositionXMS, ClampedPositionXMS + RangeXMS, true );
-	}
+	void SetPositionXToByScrollBar( double ScrollOffset );
+	void SetPositionX( double NewPositionXMS );
 
 	/**
 	 *	Changes the position-y of the thread view.
 	 *	Called by the external code.
 	 */
-	void SetPositonYTo( double ScrollOffset )
-	{
-
-	}
+	void SetPositonYTo( double ScrollOffset );
 
 	/** Changes the position-x and range-x of the thread view. */
-	void SetTimeRange( double StartTimeMS, double EndTimeMS, bool bBroadcast = true )
-	{
-		check( EndTimeMS > StartTimeMS );
-
-		PositionXMS = StartTimeMS;
-		RangeXMS = EndTimeMS - StartTimeMS;
-		FramesIndices = ProfilerStream->GetFramesIndicesForTimeRange( StartTimeMS, EndTimeMS );
-
-		bUpdateData = true;
-		
-		//UE_LOG( LogTemp, Log, TEXT( "StartTimeMS=%f, EndTimeMS=%f, bBroadcast=%1i FramesIndices=%3i,%3i" ), StartTimeMS, EndTimeMS, (int)bBroadcast, FramesIndices.X, FramesIndices.Y );
-		if( bBroadcast )
-		{
-			ViewPositionXChangedEvent.Broadcast( StartTimeMS, EndTimeMS, TotalRangeXMS, FramesIndices.X, FramesIndices.Y );
-		}
-	}
+	void SetTimeRange( double StartTimeMS, double EndTimeMS, bool bBroadcast = true );
 
 	/**
 	 *	Changes the position-x and range-x of the thread view. 
 	 *	Called by the mini-view.
 	 */
-	void SetFrameRange( int32 FrameStart, int32 FrameEnd )
-	{
-		const double EndTimeMS = ProfilerStream->GetElapsedFrameTimeMS( FrameEnd );
-		const double StartTimeMS = ProfilerStream->GetElapsedFrameTimeMS( FrameStart ) - ProfilerStream->GetFrameTimeMS( FrameStart );
-		SetTimeRange( StartTimeMS, EndTimeMS, true );
-	}
+	void SetFrameRange( int32 FrameStart, int32 FrameEnd );
 
 	/** Attaches profiler stream to the thread-view widgets and displays the first frame of data. */
-	void AttachProfilerStream( const FProfilerStream& InProfilerStream )
-	{
-		ProfilerStream = &InProfilerStream;
-
-		TotalRangeXMS = ProfilerStream->GetElapsedTime();
-		TotalRangeY = ProfilerStream->GetNumThreads()*FProfilerUIStream::DEFAULT_VISIBLE_THREAD_DEPTH;
-
-		// Display the first frame.
-		const FProfilerFrame* ProfilerFrame = ProfilerStream->GetProfilerFrame( 0 );
-		SetTimeRange( ProfilerFrame->Root->CycleCounterStartTimeMS, ProfilerFrame->Root->CycleCounterEndTimeMS );
-	}
+	void AttachProfilerStream( const FProfilerStream& InProfilerStream );
 
 public:
 	/** The event to execute when the position-x of the thread view has been changed. */
@@ -249,33 +205,15 @@ protected:
 	/** The event to execute when the position-y of the thread view has been changed. */
 	FViewPositionYChangedEvent ViewPositionYChangedEvent;
 
-
 protected:
-	void UpdateInternalConstants()
-	{
-		ZoomFactorX = (double)NUM_MILLISECONDS_PER_WINDOW / RangeXMS;
-		RangeY = FMath::RoundToFloat( ThisGeometry.Size.Y / (double)NUM_PIXELS_PER_ROW );
-		
-		const double Aspect = ThisGeometry.Size.X / NUM_MILLISECONDS_PER_WINDOW * ZoomFactorX;
-		NumMillisecondsPerWindow = (double)ThisGeometry.Size.X / Aspect;
-		NumPixelsPerMillisecond = (double)ThisGeometry.Size.X / NumMillisecondsPerWindow;
-		NumMillisecondsPerSample = NumMillisecondsPerWindow / (double)ThisGeometry.Size.X * (double)MIN_NUM_PIXELS_PER_SAMPLE;	
-	}
 	
-	void ProcessData();
-
 	/**
 	 * @return True, if the widget is ready to use, also means that contains at least one frame of the thread data.
 	 */
-	bool IsReady() const
-	{
-		return ProfilerStream && ProfilerStream->GetNumFrames() > 0;
-	}
-
-	bool ShouldUpdateData()
-	{
-		return bUpdateData;
-	}
+	bool IsReady() const;
+	void ProcessData();
+	bool ShouldUpdateData();
+	void UpdateInternalConstants();
 
 protected:
 
diff --git a/Engine/Source/Developer/Profiler/Private/Widgets/SProfilerToolbar.cpp b/Engine/Source/Developer/Profiler/Private/Widgets/SProfilerToolbar.cpp
--- a/Engine/Source/Developer/Profiler/Private/Widgets/SProfilerToolbar.cpp
+++ b/Engine/Source/Developer/Profiler/Private/Widgets/SProfilerToolbar.cpp
@@ -1,24 +1,31 @@
 // Copyright 1998-2016 Epic Games, Inc. All Rights Reserved.
 
 #include "ProfilerPrivatePCH.h"
-#include "SDockTab.h"
+#include "ProfilerManager.h"
+#include "ProfilerSession.h"
+#include "SProfilerFPSChartPanel.h"
+#include "SProfilerToolbar.h"
+
 
 #define LOCTEXT_NAMESPACE "Profiler"
 
+
 SProfilerToolbar::SProfilerToolbar()
 {
 }
 
+
 SProfilerToolbar::~SProfilerToolbar()
 {
 	// Remove ourselves from the profiler manager.
-	if( FProfilerManager::Get().IsValid() )
+	if(FProfilerManager::Get().IsValid())
 	{
-		FProfilerManager::Get()->OnSessionInstancesUpdated().RemoveAll( this );
+		FProfilerManager::Get()->OnSessionInstancesUpdated().RemoveAll(this);
 	}
 }
 
-void SProfilerToolbar::Construct( const FArguments& InArgs )
+
+void SProfilerToolbar::Construct(const FArguments& InArgs)
 {
 	CreateCommands();
 
@@ -30,39 +37,39 @@ void SProfilerToolbar::Construct( const FArguments& InArgs )
 			{
 				ToolbarBuilder.AddToolBarButton(FProfilerCommands::Get().ProfilerManager_Load);
 				ToolbarBuilder.AddToolBarButton(FProfilerCommands::Get().ProfilerManager_LoadMultiple);
-				ToolbarBuilder.AddToolBarButton( FProfilerCommands::Get().ProfilerManager_Save );
+				ToolbarBuilder.AddToolBarButton(FProfilerCommands::Get().ProfilerManager_Save);
 			}
 			ToolbarBuilder.EndSection();
 			ToolbarBuilder.BeginSection("Capture");
 			{
-				ToolbarBuilder.AddToolBarButton( FProfilerCommands::Get().ToggleDataPreview );
-				ToolbarBuilder.AddToolBarButton( FProfilerCommands::Get().ProfilerManager_ToggleLivePreview );
-				ToolbarBuilder.AddToolBarButton( FProfilerCommands::Get().ToggleDataCapture );
+				ToolbarBuilder.AddToolBarButton(FProfilerCommands::Get().ToggleDataPreview);
+				ToolbarBuilder.AddToolBarButton(FProfilerCommands::Get().ProfilerManager_ToggleLivePreview);
+				ToolbarBuilder.AddToolBarButton(FProfilerCommands::Get().ToggleDataCapture);
 			}
 			ToolbarBuilder.EndSection();
 			ToolbarBuilder.BeginSection("Profilers");
 			{
-				ToolbarBuilder.AddToolBarButton( FProfilerCommands::Get().StatsProfiler );
-				//ToolbarBuilder.AddToolBarButton( FProfilerCommands::Get().MemoryProfiler );
-				ToolbarBuilder.AddToolBarButton( FProfilerCommands::Get().FPSChart );
+				ToolbarBuilder.AddToolBarButton(FProfilerCommands::Get().StatsProfiler);
+				//ToolbarBuilder.AddToolBarButton(FProfilerCommands::Get().MemoryProfiler);
+				ToolbarBuilder.AddToolBarButton(FProfilerCommands::Get().FPSChart);
 			}
 			ToolbarBuilder.EndSection();
 			ToolbarBuilder.BeginSection("Options");
 			{
-				ToolbarBuilder.AddToolBarButton( FProfilerCommands::Get().OpenSettings );
+				ToolbarBuilder.AddToolBarButton(FProfilerCommands::Get().OpenSettings);
 			}
 			ToolbarBuilder.EndSection();
 		}
 	};
 
 	TSharedPtr<FUICommandList> ProfilerCommandList = FProfilerManager::Get()->GetCommandList();
-	FToolBarBuilder ToolbarBuilder( ProfilerCommandList.ToSharedRef(), FMultiBoxCustomization::None );
-	Local::FillToolbar( ToolbarBuilder );
+	FToolBarBuilder ToolbarBuilder(ProfilerCommandList.ToSharedRef(), FMultiBoxCustomization::None);
+	Local::FillToolbar(ToolbarBuilder);
 
 	// Create the tool bar!
 	ChildSlot
 	[
-		SNew( SHorizontalBox )
+		SNew(SHorizontalBox)
 
 		+SHorizontalBox::Slot()
 		.HAlign(HAlign_Left)
@@ -70,10 +77,10 @@ void SProfilerToolbar::Construct( const FArguments& InArgs )
 		.FillWidth(1.0)
 		.Padding(0.0f)
 		[
-			SNew( SBorder )
+			SNew(SBorder)
 			.Padding(0)
-			.BorderImage( FEditorStyle::GetBrush("NoBorder") )
-			.IsEnabled( FSlateApplication::Get().GetNormalExecutionAttribute() )
+			.BorderImage(FEditorStyle::GetBrush("NoBorder"))
+			.IsEnabled(FSlateApplication::Get().GetNormalExecutionAttribute())
 			[
 				ToolbarBuilder.MakeWidget()
 			]
@@ -81,22 +88,25 @@ void SProfilerToolbar::Construct( const FArguments& InArgs )
 	];
 }
 
+
 void SProfilerToolbar::ShowStats()
 {
-
+	// do nothing
 }
 
+
 void SProfilerToolbar::ShowMemory()
 {
-
+	// do nothing
 }
 
-void DisplayFPSChart( const TSharedRef<FFPSAnalyzer> InFPSAnalyzer )
+
+void DisplayFPSChart(const TSharedRef<FFPSAnalyzer> InFPSAnalyzer)
 {
 	static bool HasRegisteredFPSChart = false;
-	if ( !HasRegisteredFPSChart )
+	if (!HasRegisteredFPSChart)
 	{
-		TSharedRef<FTabManager::FLayout> Layout = FTabManager::NewLayout( "FPSChart_Layout" )
+		TSharedRef<FTabManager::FLayout> Layout = FTabManager::NewLayout("FPSChart_Layout")
 			->AddArea
 			(
 			FTabManager::NewArea(720, 360)
@@ -107,62 +117,65 @@ void DisplayFPSChart( const TSharedRef<FFPSAnalyzer> InFPSAnalyzer )
 			)
 			);
 
-		FGlobalTabmanager::Get()->RestoreFrom( Layout, TSharedPtr<SWindow>() );
+		FGlobalTabmanager::Get()->RestoreFrom(Layout, TSharedPtr<SWindow>());
 		HasRegisteredFPSChart = true;
 	}
 
 	FGlobalTabmanager::Get()->InsertNewDocumentTab
 	(
 		"FPSChart", FTabManager::ESearchPreference::RequireClosedTab,
-		SNew( SDockTab )
-		.Label( LOCTEXT("Label_FPSHistogram", "FPS Histogram") )
-		.TabRole( ETabRole::DocumentTab )
+		SNew(SDockTab)
+		.Label(LOCTEXT("Label_FPSHistogram", "FPS Histogram"))
+		.TabRole(ETabRole::DocumentTab)
 		[
-			SNew( SProfilerFPSChartPanel )
-			.FPSAnalyzer( InFPSAnalyzer )
+			SNew(SProfilerFPSChartPanel)
+			.FPSAnalyzer(InFPSAnalyzer)
 		]
 	);
 }
 
+
 void SProfilerToolbar::ShowFPSChart()
 {
-	const FProfilerSessionPtr ProfilerSession = FProfilerManager::Get()->GetProfilerSession();
+	const TSharedPtr<FProfilerSession> ProfilerSession = FProfilerManager::Get()->GetProfilerSession();
 	if (ProfilerSession.IsValid())
 	{
-		DisplayFPSChart( ProfilerSession->FPSAnalyzer );
+		DisplayFPSChart(ProfilerSession->FPSAnalyzer);
 	}
 }
 
+
 void SProfilerToolbar::CreateCommands()
 {
 	TSharedPtr<FUICommandList> ProfilerCommandList = FProfilerManager::Get()->GetCommandList();
 	const FProfilerCommands& Commands = FProfilerCommands::Get();
 
 	// Save command
-	ProfilerCommandList->MapAction( Commands.ProfilerManager_Save,
+	ProfilerCommandList->MapAction(Commands.ProfilerManager_Save,
 		FExecuteAction(),
-		FCanExecuteAction::CreateRaw( this, &SProfilerToolbar::IsImplemented )
+		FCanExecuteAction::CreateRaw(this, &SProfilerToolbar::IsImplemented)
 	);
 
 	// Stats command
-	ProfilerCommandList->MapAction( Commands.StatsProfiler,
-		FExecuteAction::CreateRaw( this, &SProfilerToolbar::ShowStats ),
+	ProfilerCommandList->MapAction(Commands.StatsProfiler,
+		FExecuteAction::CreateRaw(this, &SProfilerToolbar::ShowStats),
 		FCanExecuteAction(),
-		FIsActionChecked::CreateRaw( this, &SProfilerToolbar::IsShowingStats )
+		FIsActionChecked::CreateRaw(this, &SProfilerToolbar::IsShowingStats)
 	);
 
 	// Memory command
-	ProfilerCommandList->MapAction( Commands.MemoryProfiler,
-		FExecuteAction::CreateRaw( this, &SProfilerToolbar::ShowMemory ),
+	ProfilerCommandList->MapAction(Commands.MemoryProfiler,
+		FExecuteAction::CreateRaw(this, &SProfilerToolbar::ShowMemory),
 		FCanExecuteAction(),
-		FIsActionChecked::CreateRaw( this, &SProfilerToolbar::IsShowingMemory )
+		FIsActionChecked::CreateRaw(this, &SProfilerToolbar::IsShowingMemory)
 		);
 
 	// FPSChart command
-	ProfilerCommandList->MapAction( Commands.FPSChart,
-		FExecuteAction::CreateRaw( this, &SProfilerToolbar::ShowFPSChart ),
+	ProfilerCommandList->MapAction(Commands.FPSChart,
+		FExecuteAction::CreateRaw(this, &SProfilerToolbar::ShowFPSChart),
 		FCanExecuteAction()
 		);
 }
 
+
 #undef LOCTEXT_NAMESPACE
diff --git a/Engine/Source/Developer/Profiler/Private/Widgets/SProfilerWindow.cpp b/Engine/Source/Developer/Profiler/Private/Widgets/SProfilerWindow.cpp
--- a/Engine/Source/Developer/Profiler/Private/Widgets/SProfilerWindow.cpp
+++ b/Engine/Source/Developer/Profiler/Private/Widgets/SProfilerWindow.cpp
@@ -1,20 +1,22 @@
 ﻿// Copyright 1998-2016 Epic Games, Inc. All Rights Reserved.
 
 #include "ProfilerPrivatePCH.h"
+#include "ProfilerManager.h"
+#include "SDataGraph.h"
+#include "SEventGraph.h"
+#include "SFiltersAndPresets.h"
+#include "SMultiDumpBrowser.h"
+#include "SProfilerGraphPanel.h"
+#include "SProfilerMiniView.h"
 #include "SProfilerSettings.h"
-#include "DesktopPlatformModule.h"
-
-#if WITH_EDITOR
-	#include "Runtime/Analytics/Analytics/Public/Interfaces/IAnalyticsProvider.h"
-	#include "Runtime/Engine/Public/EngineAnalytics.h"
-#endif // WITH_EDITOR
-#include "SNotificationList.h"
+#include "SProfilerToolbar.h"
+#include "SProfilerWindow.h"
 
 
 #define LOCTEXT_NAMESPACE "SProfilerWindow"
 
 
-static FText GetTextForNotification( const EProfilerNotificationTypes::Type NotificatonType, const ELoadingProgressStates::Type ProgressState, const FString& Filename, const float ProgressPercent = 0.0f )
+static FText GetTextForNotification( const EProfilerNotificationTypes NotificatonType, const ELoadingProgressStates ProgressState, const FString& Filename, const float ProgressPercent = 0.0f )
 {
 	FText Result;
 
@@ -398,7 +400,8 @@ bool SProfilerWindow::IsProfilerEnabled() const
 	return bIsActive;
 }
 
-void SProfilerWindow::ManageLoadingProgressNotificationState( const FString& Filename, const EProfilerNotificationTypes::Type NotificatonType, const ELoadingProgressStates::Type ProgressState, const float DataLoadingProgress )
+
+void SProfilerWindow::ManageLoadingProgressNotificationState( const FString& Filename, const EProfilerNotificationTypes NotificatonType, const ELoadingProgressStates ProgressState, const float DataLoadingProgress )
 {
 	const FString BaseFilename = FPaths::GetBaseFilename( Filename );
 
@@ -629,7 +632,7 @@ void SProfilerWindow::CloseProfilerSettings()
 	MainContentPanel->SetEnabled( true );
 }
 
-void SProfilerWindow::ProfilerManager_OnViewModeChanged( EProfilerViewMode::Type NewViewMode )
+void SProfilerWindow::ProfilerManager_OnViewModeChanged( EProfilerViewMode NewViewMode )
 {
 	if( NewViewMode == EProfilerViewMode::LineIndexBased )
 	{
diff --git a/Engine/Source/Developer/Profiler/Private/Widgets/SProfilerWindow.h b/Engine/Source/Developer/Profiler/Private/Widgets/SProfilerWindow.h
--- a/Engine/Source/Developer/Profiler/Private/Widgets/SProfilerWindow.h
+++ b/Engine/Source/Developer/Profiler/Private/Widgets/SProfilerWindow.h
@@ -2,9 +2,16 @@
 
 #pragma once
 
-/*-----------------------------------------------------------------------------
-	Type definitions
------------------------------------------------------------------------------*/
+#include "ProfilerDataSource.h" // @todo profiler: get rid of ProfilerDataSource typedefs
+#include "ProfilerSession.h"
+
+
+enum class EProfilerViewMode;
+class SFiltersAndPresets;
+class SProfilerMiniView;
+class SMultiDumpBrowser;
+class SProfilerGraphPanel;
+
 
 /** Type definition for shared pointers to instances of SNotificationItem. */
 typedef TSharedPtr<class SNotificationItem> SNotificationItemPtr;
@@ -15,9 +22,6 @@ typedef TSharedRef<class SNotificationItem> SNotificationItemRef;
 /** Type definition for weak references to instances of SNotificationItem. */
 typedef TWeakPtr<class SNotificationItem> SNotificationItemWeak;
 
-/*-----------------------------------------------------------------------------
-	SProfilerWindow class
------------------------------------------------------------------------------*/
 
 /**
  * Implements the profiler window.
@@ -42,7 +46,7 @@ public:
 	void ManageEventGraphTab( const FGuid ProfilerInstanceID, const bool bCreateFakeTab, const FString TabName );
 	void UpdateEventGraph( const FGuid ProfilerInstanceID, const FEventGraphDataRef AverageEventGraph, const FEventGraphDataRef MaximumEventGraph, bool bInitial );
 
-	void ManageLoadingProgressNotificationState( const FString& Filename, const EProfilerNotificationTypes::Type NotificatonType, const ELoadingProgressStates::Type ProgressState, const float DataLoadingProgress );
+	void ManageLoadingProgressNotificationState( const FString& Filename, const EProfilerNotificationTypes NotificatonType, const ELoadingProgressStates ProgressState, const float DataLoadingProgress );
 
 	void OpenProfilerSettings();
 	void CloseProfilerSettings();
@@ -58,7 +62,7 @@ protected:
 
 	void SendingServiceSideCapture_Load( const FString Filename );
 
-	void ProfilerManager_OnViewModeChanged( EProfilerViewMode::Type NewViewMode );
+	void ProfilerManager_OnViewModeChanged( EProfilerViewMode NewViewMode );
 
 private:
 
diff --git a/Engine/Source/Developer/Profiler/Profiler.Build.cs b/Engine/Source/Developer/Profiler/Profiler.Build.cs
--- a/Engine/Source/Developer/Profiler/Profiler.Build.cs
+++ b/Engine/Source/Developer/Profiler/Profiler.Build.cs
@@ -49,13 +49,5 @@ public class Profiler : ModuleRules
 				"SessionServices",
 			}
 		);
-
-		PublicIncludePaths.AddRange
-		(
-			new string[] {
-				"Developer/Profiler/Public",
-				"Developer/Profiler/Public/Interfaces",
-			}
-		);
 	}
 }
diff --git a/Engine/Source/Developer/Profiler/Public/IProfilerModule.h b/Engine/Source/Developer/Profiler/Public/IProfilerModule.h
new file mode 100644
--- /dev/null
+++ b/Engine/Source/Developer/Profiler/Public/IProfilerModule.h
@@ -0,0 +1,37 @@
+// Copyright 1998-2016 Epic Games, Inc. All Rights Reserved.
+
+#pragma once
+
+
+class FRawStatsMemoryProfiler;
+class ISessionManager;
+class SDockTab;
+class SWidget;
+
+
+/**
+ * Interface for the profiler module.
+ */
+class IProfilerModule
+	: public IModuleInterface
+{
+public:
+
+	/**
+	 * Creates the main window for the profiler.
+	 *
+	 * @param InSessionManager The session manager to use.
+	 * @param ConstructUnderMajorTab The tab which will contain the profiler tabs.
+	 *
+	 */
+	virtual TSharedRef<SWidget> CreateProfilerWindow(const TSharedRef<ISessionManager>& InSessionManager, const TSharedRef<SDockTab>& ConstructUnderMajorTab) = 0;
+
+	/** Implements stats memory dump command. */
+	virtual void StatsMemoryDumpCommand(const TCHAR* Filename) = 0;
+
+	/** 
+	 * Creates a new instance of the memory profiler based the raw stats. 
+	 * When no longer needed must be stopped via RequestStop() and deleted to avoid memory leaks.
+	 */
+	virtual FRawStatsMemoryProfiler* OpenRawStatsForMemoryProfiling(const TCHAR* Filename) = 0;
+};
diff --git a/Engine/Source/Developer/Profiler/Public/ProfilerCommon.h b/Engine/Source/Developer/Profiler/Public/ProfilerCommon.h
--- a/Engine/Source/Developer/Profiler/Public/ProfilerCommon.h
+++ b/Engine/Source/Developer/Profiler/Public/ProfilerCommon.h
@@ -1,62 +1,20 @@
 // Copyright 1998-2016 Epic Games, Inc. All Rights Reserved.
 
-/*=============================================================================
-	ProfilerCommon.h: Declares the profiler common classes, structures, typedef etc.
-=============================================================================*/
-
 #pragma once
-DECLARE_LOG_CATEGORY_EXTERN(Profiler, Log, All);
-
-/*-----------------------------------------------------------------------------
-	Type definitions
------------------------------------------------------------------------------*/
-
-/** Type definition for shared pointers to instances of FGraphDataSource. */
-typedef TSharedPtr<class FGraphDataSource> FGraphDataSourcePtr;
-
-/** Type definition for shared references to const instances of FGraphDataSource. */
-typedef TSharedRef<const class FGraphDataSource> FGraphDataSourceRefConst;
-
-/** Type definition for shared pointers to instances of FCombinedGraphDataSource. */
-typedef TSharedPtr<class FCombinedGraphDataSource> FCombinedGraphDataSourcePtr;
-
-/** Type definition for shared references to instances of FCombinedGraphDataSource. */
-typedef TSharedRef<class FCombinedGraphDataSource> FCombinedGraphDataSourceRef;
-
-/** Type definition for shared pointers to instances of IDataProvider. */
-typedef TSharedPtr<class IDataProvider> IDataProviderPtr;
-
-/** Type definition for shared references to instances of IDataProvider. */
-typedef TSharedRef<class IDataProvider> IDataProviderRef;
 
-/** Type definition for shared pointers to instances of FProfilerSession. */
-typedef TSharedPtr<class FProfilerSession> FProfilerSessionPtr;
 
-/** Type definition for shared references to instances of FProfilerSession. */
-typedef TSharedRef<class FProfilerSession> FProfilerSessionRef;
-
-/** Type definition for shared pointers to instances of FProfilerStatMetaData. */
-typedef TSharedPtr<class FProfilerStatMetaData> FProfilerStatMetaDataPtr;
-
-/** Type definition for shared references to instances of FProfilerStatMetaData. */
-typedef TSharedRef<class FProfilerStatMetaData> FProfilerStatMetaDataRef;
-
-/** Type definition for shared pointers to instances of FTrackedStat. */
-typedef TSharedPtr<class FTrackedStat> FTrackedStatPtr;
+DECLARE_LOG_CATEGORY_EXTERN(Profiler, Log, All);
 
-class FProfilerStat;
-class FProfilerGroup;
-class FProfilerStatMetaData;
-class FProfilerAggregatedStat;
 
 #define DEBUG_PROFILER_PERFORMANCE 0
 
 #if DEBUG_PROFILER_PERFORMANCE==1
-#define PROFILER_SCOPE_LOG_TIME(arg0,arg1) SCOPE_LOG_TIME(arg0,arg1)
+	#define PROFILER_SCOPE_LOG_TIME(arg0,arg1) SCOPE_LOG_TIME(arg0,arg1)
 #else
-#define PROFILER_SCOPE_LOG_TIME(arg0,arg1)
+	#define PROFILER_SCOPE_LOG_TIME(arg0,arg1)
 #endif
 
+
 /** Time spent on graph drawing. */
 DECLARE_CYCLE_STAT_EXTERN( TEXT("DataGraphOnPaint"),	STAT_DG_OnPaint,			STATGROUP_Profiler, );
 
@@ -69,30 +27,23 @@ DECLARE_CYCLE_STAT_EXTERN( TEXT("ProfilerTick"),		STAT_PM_Tick,				STATGROUP_Pro
 /** Number of bytes allocated by all profiler sessions. */
 DECLARE_MEMORY_STAT_EXTERN( TEXT("ProfilerMemoryUsage"),STAT_PM_MemoryUsage,		STATGROUP_Profiler, );
 
-/*-----------------------------------------------------------------------------
-	Enumerators
------------------------------------------------------------------------------*/
 
-/** Enumerates graph styles. */
-namespace EGraphStyles
+/**
+ * Enumerates graph styles.
+ */
+enum class EGraphStyles
 {
-	enum Type
-	{
-		/** Line graph. */
-		Line,
+	/** Line graph. */
+	Line,
 
-		/** Combined graph. */
-		Combined,
+	/** Combined graph. */
+	Combined,
 
-		/** More to come in future. */
-		/** Invalid enum type, may be used as a number of enumerations. */
-		InvalidOrMax,
-	};
-}
+	/** More to come in future. */
+	/** Invalid enum type, may be used as a number of enumerations. */
+	InvalidOrMax,
+};
 
-/*-----------------------------------------------------------------------------
-	Basic structures
------------------------------------------------------------------------------*/
 
 class FProfilerHelper
 {
@@ -109,7 +60,10 @@ public:
 	}
 };
 
-/** Implements binary search for the various cases. */
+
+/**
+ * Implements binary search for the various cases.
+ */
 class FBinaryFindIndex
 {
 public:
@@ -179,4 +133,4 @@ public:
 		}
 		return Offset+Edge;
 	}
-};
\ No newline at end of file
+};
