Fixed crash when changing "Show Coalesced" setting in profiler

#rnx

diff --git a/Engine/Source/Developer/Profiler/Private/Widgets/SEventGraph.cpp b/Engine/Source/Developer/Profiler/Private/Widgets/SEventGraph.cpp
--- a/Engine/Source/Developer/Profiler/Private/Widgets/SEventGraph.cpp
+++ b/Engine/Source/Developer/Profiler/Private/Widgets/SEventGraph.cpp
@@ -3359,7 +3359,7 @@ EVisibility SEventGraph::EventGraphViewMode_GetVisibility( const EEventGraphView
 	if( ViewMode == EEventGraphViewModes::FlatInclusiveCoalesced || ViewMode == EEventGraphViewModes::FlatExclusiveCoalesced )
 	{
 		const EVisibility Vis = FProfilerManager::GetSettings().bShowCoalescedViewModesInEventGraph ? EVisibility::Visible : EVisibility::Collapsed;
-		if( Vis == EVisibility::Collapsed )
+		if (Vis == EVisibility::Collapsed && GetCurrentStateViewMode() == ViewMode)
 		{
 			// If view mode is not available event graph will switch to the hierarchical view mode.
 			SEventGraph* MutableThis = const_cast< SEventGraph* >( this );
