PR #5753: Fix Wording (Contributed by [REDACTED_PERSON])

(peafour-cherrypick of //UE4/Main/[REDACTED_HANDLE] by UnrealBot)

diff --git a/Engine/Source/Developer/Profiler/Private/ProfilerCommands.cpp b/Engine/Source/Developer/Profiler/Private/ProfilerCommands.cpp
--- a/Engine/Source/Developer/Profiler/Private/ProfilerCommands.cpp
+++ b/Engine/Source/Developer/Profiler/Private/ProfilerCommands.cpp
@@ -302,7 +302,7 @@ void FProfilerActionManager::ToggleDataCapture_Execute( const FGuid SessionInsta
 		EAppReturnType::Type Result = FPlatformMisc::MessageBoxExt
 		( 
 			EAppMsgType::YesNo, 
-			*LOCTEXT("TransferServiceSideCaptureQuestion", "Would like to transfer the captured stats file(s) to this machine? This may take some time.").ToString(), 
+			*LOCTEXT("TransferServiceSideCaptureQuestion", "Would you like to transfer the captured stats file(s) to this machine? This may take some time.").ToString(), 
 			*LOCTEXT("Question", "Question").ToString() 
 		);
 
