More const incorrectness delegate fixes.

diff --git a/Engine/Source/Developer/Profiler/Private/ProfilerCommands.cpp b/Engine/Source/Developer/Profiler/Private/ProfilerCommands.cpp
--- a/Engine/Source/Developer/Profiler/Private/ProfilerCommands.cpp
+++ b/Engine/Source/Developer/Profiler/Private/ProfilerCommands.cpp
@@ -284,7 +284,7 @@ void FProfilerActionManager::Map_ToggleDataCapture_Global()
 const FUIAction FProfilerActionManager::ToggleDataCapture_Custom( const FGuid SessionInstanceID ) const
 {
 	FUIAction UIAction;
-	UIAction.ExecuteAction = FExecuteAction::CreateRaw( this, &FProfilerActionManager::ToggleDataCapture_Execute, SessionInstanceID );
+	UIAction.ExecuteAction = FExecuteAction::CreateRaw( const_cast<FProfilerActionManager*>(this), &FProfilerActionManager::ToggleDataCapture_Execute, SessionInstanceID );
 	UIAction.CanExecuteAction = FCanExecuteAction::CreateRaw( this, &FProfilerActionManager::ToggleDataCapture_CanExecute, SessionInstanceID );
 	UIAction.GetActionCheckState = FGetActionCheckState::CreateRaw( this, &FProfilerActionManager::ToggleDataCapture_GetCheckState, SessionInstanceID );
 	return UIAction;
diff --git a/Engine/Source/Developer/Profiler/Private/Widgets/SEventGraph.cpp b/Engine/Source/Developer/Profiler/Private/Widgets/SEventGraph.cpp
--- a/Engine/Source/Developer/Profiler/Private/Widgets/SEventGraph.cpp
+++ b/Engine/Source/Developer/Profiler/Private/Widgets/SEventGraph.cpp
@@ -1735,7 +1735,7 @@ TSharedPtr<SWidget> SEventGraph::EventGraph_GetMenuContent() const
 
 		FUIAction Action_ExpandHotPath
 		(
-			FExecuteAction::CreateSP( this, &SEventGraph::ContextMenu_ExpandHotPath_Execute ),
+			FExecuteAction::CreateSP( const_cast<SEventGraph*>(this), &SEventGraph::ContextMenu_ExpandHotPath_Execute ),
 			FCanExecuteAction::CreateSP( this, &SEventGraph::ContextMenu_ExpandHotPath_CanExecute )
 		);
 		MenuBuilder.AddMenuEntry
@@ -1812,7 +1812,7 @@ TSharedPtr<SWidget> SEventGraph::EventGraph_GetMenuContent() const
 	{
 		FUIAction Action_CopySelectedToClipboard
 		(
-			FExecuteAction::CreateSP( this, &SEventGraph::ContextMenu_CopySelectedToClipboard_Execute ),
+			FExecuteAction::CreateSP( const_cast<SEventGraph*>(this), &SEventGraph::ContextMenu_CopySelectedToClipboard_Execute ),
 			FCanExecuteAction::CreateSP( this, &SEventGraph::ContextMenu_CopySelectedToClipboard_CanExecute )
 		);
 		MenuBuilder.AddMenuEntry
@@ -1826,7 +1826,7 @@ TSharedPtr<SWidget> SEventGraph::EventGraph_GetMenuContent() const
 
 		FUIAction Action_SelectStack
 		(
-			FExecuteAction::CreateSP( this, &SEventGraph::ContextMenu_SelectStack_Execute ),
+			FExecuteAction::CreateSP( const_cast<SEventGraph*>(this), &SEventGraph::ContextMenu_SelectStack_Execute ),
 			FCanExecuteAction::CreateSP( this, &SEventGraph::ContextMenu_SelectStack_CanExecute )
 		);
 		MenuBuilder.AddMenuEntry
@@ -1853,7 +1853,7 @@ TSharedPtr<SWidget> SEventGraph::EventGraph_GetMenuContent() const
 		( 
 			LOCTEXT("ContextMenu_Header_Columns_View", "View Column"), 
 			LOCTEXT("ContextMenu_Header_Columns_View_Desc", "Hides or shows columns"), 
-			FNewMenuDelegate::CreateSP( this, &SEventGraph::EventGraph_BuildViewColumnMenu ),
+			FNewMenuDelegate::CreateSP( const_cast<SEventGraph*>(this), &SEventGraph::EventGraph_BuildViewColumnMenu ),
 			false,
 			FSlateIcon(FEditorStyle::GetStyleSetName(), "Profiler.EventGraph.ViewColumn")
 		);
diff --git a/Engine/Source/Developer/Profiler/Private/Widgets/SEventGraph.h b/Engine/Source/Developer/Profiler/Private/Widgets/SEventGraph.h
--- a/Engine/Source/Developer/Profiler/Private/Widgets/SEventGraph.h
+++ b/Engine/Source/Developer/Profiler/Private/Widgets/SEventGraph.h
@@ -384,7 +384,7 @@ public:
 	const FUIAction SetRoot_Custom() const
 	{
 		FUIAction UIAction;
-		UIAction.ExecuteAction = FExecuteAction::CreateSP( this, &SEventGraph::SetRoot_Execute );
+		UIAction.ExecuteAction = FExecuteAction::CreateSP( const_cast<SEventGraph*>(this), &SEventGraph::SetRoot_Execute );
 		UIAction.CanExecuteAction = FCanExecuteAction::CreateSP( this, &SEventGraph::SetRoot_CanExecute );
 		return UIAction;
 	}
@@ -405,7 +405,7 @@ public:
 	const FUIAction ClearHistory_Custom() const
 	{
 		FUIAction UIAction;
-		UIAction.ExecuteAction = FExecuteAction::CreateSP( this, &SEventGraph::ClearHistory_Execute );
+		UIAction.ExecuteAction = FExecuteAction::CreateSP( const_cast<SEventGraph*>(this), &SEventGraph::ClearHistory_Execute );
 		UIAction.CanExecuteAction = FCanExecuteAction::CreateSP( this, &SEventGraph::ClearHistory_CanExecute );
 		return UIAction;
 	}
@@ -469,7 +469,7 @@ public:
 	const FUIAction CullByProperty_Custom( const FEventGraphSamplePtr EventPtr, const FName PropertyName, const bool bReset ) const
 	{
 		FUIAction UIAction;
-		UIAction.ExecuteAction = FExecuteAction::CreateSP( this, &SEventGraph::CullByProperty_Execute, EventPtr, PropertyName, bReset );
+		UIAction.ExecuteAction = FExecuteAction::CreateSP( const_cast<SEventGraph*>(this), &SEventGraph::CullByProperty_Execute, EventPtr, PropertyName, bReset );
 		UIAction.CanExecuteAction = FCanExecuteAction::CreateSP( this, &SEventGraph::CullByProperty_CanExecute, EventPtr, PropertyName, bReset );
 		return UIAction;
 	}
diff --git a/Engine/Source/Editor/LevelEditor/Private/SLevelViewportToolBar.cpp b/Engine/Source/Editor/LevelEditor/Private/SLevelViewportToolBar.cpp
--- a/Engine/Source/Editor/LevelEditor/Private/SLevelViewportToolBar.cpp
+++ b/Engine/Source/Editor/LevelEditor/Private/SLevelViewportToolBar.cpp
@@ -1152,7 +1152,7 @@ TSharedRef<SWidget> SLevelViewportToolBar::GenerateFarViewPlaneMenu() const
 				.MaxValue(100000.0f)
 				.Font(FEditorStyle::GetFontStyle(TEXT("MenuItem.Font")))
 				.Value(this, &SLevelViewportToolBar::OnGetFarViewPlaneValue)
-				.OnValueChanged(this, &SLevelViewportToolBar::OnFarViewPlaneValueChanged)
+				.OnValueChanged(const_cast<SLevelViewportToolBar*>(this), &SLevelViewportToolBar::OnFarViewPlaneValueChanged)
 			]
 		];
 }
diff --git a/Engine/Source/Editor/StaticMeshEditor/Private/StaticMeshEditorTools.cpp b/Engine/Source/Editor/StaticMeshEditor/Private/StaticMeshEditorTools.cpp
--- a/Engine/Source/Editor/StaticMeshEditor/Private/StaticMeshEditorTools.cpp
+++ b/Engine/Source/Editor/StaticMeshEditor/Private/StaticMeshEditorTools.cpp
@@ -4124,8 +4124,8 @@ TSharedRef<SWidget> FLevelOfDetailSettingsLayout::GetNumStreamedLODsWidget(FName
 	return SNew(SSpinBox<int32>)
 		.Font(IDetailLayoutBuilder::GetDetailFont())
 		.Value(this, &FLevelOfDetailSettingsLayout::GetNumStreamedLODs, PlatformGroupName)
-		.OnValueChanged(this, &FLevelOfDetailSettingsLayout::OnNumStreamedLODsChanged, PlatformGroupName)
-		.OnValueCommitted(this, &FLevelOfDetailSettingsLayout::OnNumStreamedLODsCommitted, PlatformGroupName)
+		.OnValueChanged(const_cast<FLevelOfDetailSettingsLayout*>(this), &FLevelOfDetailSettingsLayout::OnNumStreamedLODsChanged, PlatformGroupName)
+		.OnValueCommitted(const_cast<FLevelOfDetailSettingsLayout*>(this), &FLevelOfDetailSettingsLayout::OnNumStreamedLODsCommitted, PlatformGroupName)
 		.MinValue(-1)
 		.MaxValue(MAX_STATIC_MESH_LODS)
 		.ToolTipText(this, &FLevelOfDetailSettingsLayout::GetNumStreamedLODsTooltip)
diff --git a/Engine/Source/Editor/WorldBrowser/Private/StreamingLevels/StreamingLevelCollectionModel.cpp b/Engine/Source/Editor/WorldBrowser/Private/StreamingLevels/StreamingLevelCollectionModel.cpp
--- a/Engine/Source/Editor/WorldBrowser/Private/StreamingLevels/StreamingLevelCollectionModel.cpp
+++ b/Engine/Source/Editor/WorldBrowser/Private/StreamingLevels/StreamingLevelCollectionModel.cpp
@@ -375,7 +375,7 @@ void FStreamingLevelCollectionModel::CustomizeFileMainMenu(FMenuBuilder& InMenuB
 		InMenuBuilder.AddSubMenu( 
 			LOCTEXT("LevelsStreamingMethod", "Default Streaming Method"),
 			LOCTEXT("LevelsStreamingMethod_Tooltip", "Changes the default streaming method for a new levels"),
-			FNewMenuDelegate::CreateRaw(this, &FStreamingLevelCollectionModel::FillDefaultStreamingMethodSubMenu ) );
+			FNewMenuDelegate::CreateRaw(const_cast<FStreamingLevelCollectionModel*>(this), &FStreamingLevelCollectionModel::FillDefaultStreamingMethodSubMenu ) );
 		
 		InMenuBuilder.AddMenuEntry( Commands.World_CreateNewLevel );
 		InMenuBuilder.AddMenuEntry( Commands.World_AddExistingLevel );
