Another const_cast fix.

diff --git a/Engine/Source/Developer/Profiler/Private/Widgets/SEventGraph.cpp b/Engine/Source/Developer/Profiler/Private/Widgets/SEventGraph.cpp
--- a/Engine/Source/Developer/Profiler/Private/Widgets/SEventGraph.cpp
+++ b/Engine/Source/Developer/Profiler/Private/Widgets/SEventGraph.cpp
@@ -1860,7 +1860,7 @@ TSharedPtr<SWidget> SEventGraph::EventGraph_GetMenuContent() const
 
 		FUIAction Action_ResetColumns
 		(
-			FExecuteAction::CreateSP( this, &SEventGraph::ContextMenu_ResetColumns_Execute ),
+			FExecuteAction::CreateSP( const_cast<SEventGraph*>(this), &SEventGraph::ContextMenu_ResetColumns_Execute ),
 			FCanExecuteAction::CreateSP( this, &SEventGraph::ContextMenu_ResetColumns_CanExecute )
 		);
 		MenuBuilder.AddMenuEntry
