Silence static analysis warning

diff --git a/Engine/Source/Developer/Profiler/Private/Widgets/SEventGraph.cpp b/Engine/Source/Developer/Profiler/Private/Widgets/SEventGraph.cpp
--- a/Engine/Source/Developer/Profiler/Private/Widgets/SEventGraph.cpp
+++ b/Engine/Source/Developer/Profiler/Private/Widgets/SEventGraph.cpp
@@ -73,13 +73,10 @@ struct FEventGraphColumns
 {
 	/** Default constructor. */
 	FEventGraphColumns()
-		: NumColumns( (uint32)EEventPropertyIndex::None+1 )
 	{
 		// Make event property is initialized.
 		FEventGraphSample::InitializePropertyManagement();
 
-		Collection = new FEventGraphColumn[NumColumns];
-
 		Collection[(uint32)EEventPropertyIndex::StatName] = FEventGraphColumn
 		(
 			EEventPropertyIndex::StatName,
@@ -167,17 +164,12 @@ struct FEventGraphColumns
 			.Add( TEXT( "NumCallsPerFrame" ), &Collection[(uint32)EEventPropertyIndex::NumCallsPerFrame] )
 			.Add( NAME_None, &Collection[(uint32)EEventPropertyIndex::None] )
 			;
-	}
+	}	
 
-	~FEventGraphColumns()
-	{
-		delete[] Collection;
-	}
+	static constexpr uint32 NumColumns = (uint32)EEventPropertyIndex::None + 1 ;
 
 	/** Contains basic information about columns used in the event graph widget. Names should be localized. */
-	FEventGraphColumn* Collection;
-
-	const uint32 NumColumns;
+	FEventGraphColumn Collection[NumColumns];
 
 	// Generated from a XLSX file.
 	TMap<FName, const FEventGraphColumn*> ColumnNameToIndexMapping;
