Replaced references to ue4stats with uestats.
Fixed frontend profiler to load uestats

diff --git a/Engine/Source/Developer/FunctionalTesting/Private/FunctionalTest.cpp b/Engine/Source/Developer/FunctionalTesting/Private/FunctionalTest.cpp
--- a/Engine/Source/Developer/FunctionalTesting/Private/FunctionalTest.cpp
+++ b/Engine/Source/Developer/FunctionalTesting/Private/FunctionalTest.cpp
@@ -1433,7 +1433,7 @@ void UAutomationPerformaceHelper::BeginStatsFile(const FString& RecordName)
 	if (UWorld* World = GetWorld())
 	{
 		FString MapName = World->GetMapName();
-		FString Cmd = FString::Printf(TEXT("Stat StartFile %s-%s/%s.ue4stats"), *MapName, *StartOfTestingTime, *RecordName);
+		FString Cmd = FString::Printf(TEXT("Stat StartFile %s-%s/%s.uestats"), *MapName, *StartOfTestingTime, *RecordName);
 		GEngine->Exec(World, *Cmd);
 	}
 }
diff --git a/Engine/Source/Developer/Profiler/Private/ProfilerCommands.cpp b/Engine/Source/Developer/Profiler/Private/ProfilerCommands.cpp
--- a/Engine/Source/Developer/Profiler/Private/ProfilerCommands.cpp
+++ b/Engine/Source/Developer/Profiler/Private/ProfilerCommands.cpp
@@ -202,7 +202,7 @@ void FProfilerActionManager::ProfilerManager_Load_Execute()
 			LOCTEXT("ProfilerManager_LoadFile_Desc", "Open profiler capture file...").ToString(),
 			ProfilingDirectory, 
 			TEXT(""), 
-			LOCTEXT("ProfilerManager_Load_FileFilter", "Stats files (*.ue4stats)|*.ue4stats|Raw Stats files (*.ue4statsraw)|*.ue4statsraw").ToString(), 
+			LOCTEXT("ProfilerManager_Load_FileFilter", "Stats files (*.uestats)|*.uestats;*.ue4stats|Raw Stats files (*.uestatsraw)|*.uestatsraw;*.ue4statsraw").ToString(), 
 			EFileDialogFlags::None,
 			OutFiles
 		);
@@ -250,7 +250,7 @@ void FProfilerActionManager::ProfilerManager_LoadMultiple_Execute()
 		if (!OutFolder.IsEmpty())
 		{
 			TArray<FString> FoundFiles;
-			FFileManagerGeneric::Get().FindFiles(FoundFiles, *OutFolder, TEXT(".ue4stats"));
+			FFileManagerGeneric::Get().FindFiles(FoundFiles, *OutFolder, TEXT(".uestats"));
 			for (FString &FilePath : FoundFiles)
 			{
 				const TCHAR* PathDelimiter = FPlatformMisc::GetDefaultPathSeparator();
diff --git a/Engine/Source/Developer/Profiler/Private/ProfilerRawStatsForMemory.h b/Engine/Source/Developer/Profiler/Private/ProfilerRawStatsForMemory.h
--- a/Engine/Source/Developer/Profiler/Private/ProfilerRawStatsForMemory.h
+++ b/Engine/Source/Developer/Profiler/Private/ProfilerRawStatsForMemory.h
@@ -322,7 +322,7 @@ public:
 	}
 
 	/**
-	 * @return Platform's name based on the loaded ue4statsraw file.
+	 * @return Platform's name based on the loaded uestatsraw file.
 	 */
 	const FString& GetPlatformName() const
 	{
diff --git a/Engine/Source/Runtime/Core/Private/Stats/StatsData.cpp b/Engine/Source/Runtime/Core/Private/Stats/StatsData.cpp
--- a/Engine/Source/Runtime/Core/Private/Stats/StatsData.cpp
+++ b/Engine/Source/Runtime/Core/Private/Stats/StatsData.cpp
@@ -26,8 +26,8 @@ const FName FStatConstants::NAME_ThreadGroup = FStatConstants::ThreadGroupName;
 const FName FStatConstants::RAW_SecondsPerCycle = FStatNameAndInfo( GET_STATFNAME( STAT_SecondsPerCycle ), true ).GetRawName();
 const FName FStatConstants::NAME_NoCategory = FName(TEXT("STATCAT_None"));
 
-const FString FStatConstants::StatsFileExtension = TEXT( ".ue4stats" );
-const FString FStatConstants::StatsFileRawExtension = TEXT( ".ue4statsraw" );
+const FString FStatConstants::StatsFileExtension = TEXT( ".uestats" );
+const FString FStatConstants::StatsFileRawExtension = TEXT( ".uestatsraw" );
 
 const FString FStatConstants::ThreadNameMarker = TEXT( "Thread_" );
 
diff --git a/Engine/Source/Runtime/Core/Private/Stats/StatsFile.cpp b/Engine/Source/Runtime/Core/Private/Stats/StatsFile.cpp
--- a/Engine/Source/Runtime/Core/Private/Stats/StatsFile.cpp
+++ b/Engine/Source/Runtime/Core/Private/Stats/StatsFile.cpp
@@ -226,7 +226,7 @@ void IStatsWriteFile::Start( const FString& InFilename )
 
 	if (FPaths::GetExtension(Filename).IsEmpty())
 	{
-		Filename += TEXT(".ue4stats");
+		Filename += TEXT(".uestats");
 	}
 
 	const FString Path = FPaths::GetPath( Filename );
diff --git a/Engine/Source/Runtime/Core/Public/Stats/StatsFile.h b/Engine/Source/Runtime/Core/Public/Stats/StatsFile.h
--- a/Engine/Source/Runtime/Core/Public/Stats/StatsFile.h
+++ b/Engine/Source/Runtime/Core/Public/Stats/StatsFile.h
@@ -941,7 +941,7 @@ struct FStatsReader
 	}
 };
 
-/** Struct used to read from ue4stats/ue4statsraw files, initializes all metadata and starts a process of reading the file asynchronously. */
+/** Struct used to read from uestats/uestatsraw files, initializes all metadata and starts a process of reading the file asynchronously. */
 struct CORE_API FStatsReadFile
 {
 	friend class FAsyncStatsFile;
