Removed remaining UE4 references in profiler code

diff --git a/Engine/Source/Developer/Profiler/Private/ProfilerCommands.cpp b/Engine/Source/Developer/Profiler/Private/ProfilerCommands.cpp
--- a/Engine/Source/Developer/Profiler/Private/ProfilerCommands.cpp
+++ b/Engine/Source/Developer/Profiler/Private/ProfilerCommands.cpp
@@ -202,7 +202,7 @@ void FProfilerActionManager::ProfilerManager_Load_Execute()
 			LOCTEXT("ProfilerManager_LoadFile_Desc", "Open profiler capture file...").ToString(),
 			ProfilingDirectory, 
 			TEXT(""), 
-			LOCTEXT("ProfilerManager_Load_FileFilter", "Stats files (*.uestats)|*.uestats;*.ue4stats|Raw Stats files (*.uestatsraw)|*.uestatsraw;*.ue4statsraw").ToString(), 
+			LOCTEXT("ProfilerManager_Load_FileFilter", "Stats files (*.uestats)|*.uestats|Raw Stats files (*.uestatsraw)|*.uestatsraw").ToString(), 
 			EFileDialogFlags::None,
 			OutFiles
 		);
