Fixed Exclusive time computation for old Profiler.

diff --git a/Engine/Source/Developer/Profiler/Private/ProfilerDataSource.h b/Engine/Source/Developer/Profiler/Private/ProfilerDataSource.h
--- a/Engine/Source/Developer/Profiler/Private/ProfilerDataSource.h
+++ b/Engine/Source/Developer/Profiler/Private/ProfilerDataSource.h
@@ -990,6 +990,7 @@ public:
 		, _MinInclusiveTimeMS( TNumericLimits<double>::Max() )
 		, _MaxInclusiveTimeMS( TNumericLimits<double>::Min() )
 		, _AvgInclusiveTimeMS( 0.0 )
+		, _AvgExclusiveTimeMS( 0.0 )
 		, _NumCallsPerFrame( 1.0 )
 		, _MinNumCallsPerFrame( TNumericLimits<double>::Max() )
 		, _MaxNumCallsPerFrame( TNumericLimits<double>::Min() )
@@ -1043,6 +1044,7 @@ protected:
 		, _MinInclusiveTimeMS( InInclusiveTimeMS )
 		, _MaxInclusiveTimeMS( InInclusiveTimeMS )
 		, _AvgInclusiveTimeMS( InInclusiveTimeMS )
+		, _AvgExclusiveTimeMS( 0.0f )
 		, _NumCallsPerFrame( InNumCallsPerFrame )
 		, _MinNumCallsPerFrame( InNumCallsPerFrame )
 		, _MaxNumCallsPerFrame( InNumCallsPerFrame )
@@ -1074,6 +1076,7 @@ protected:
 		, _MinInclusiveTimeMS( SourceEvent._MinInclusiveTimeMS )
 		, _MaxInclusiveTimeMS( SourceEvent._MaxInclusiveTimeMS )
 		, _AvgInclusiveTimeMS( SourceEvent._AvgInclusiveTimeMS )
+		, _AvgExclusiveTimeMS(SourceEvent._AvgExclusiveTimeMS)
 		, _NumCallsPerFrame( SourceEvent._NumCallsPerFrame )
 		, _MinNumCallsPerFrame( SourceEvent._MinNumCallsPerFrame )
 		, _MaxNumCallsPerFrame( SourceEvent._MaxNumCallsPerFrame )
@@ -1110,7 +1113,7 @@ public:
 		// Total
 		_InclusiveTimeMS += Other->_InclusiveTimeMS;
 		_NumCallsPerFrame += Other->_NumCallsPerFrame;
-		//_ExclusiveTimeMS += Other->_ExclusiveTimeMS;
+		_ExclusiveTimeMS += Other->_ExclusiveTimeMS;
 
 		// Min/Max
 		_MinInclusiveTimeMS = FMath::Min( _MinInclusiveTimeMS, Other->_MinInclusiveTimeMS );
@@ -1146,6 +1149,7 @@ protected:
 	FORCEINLINE_DEBUGGABLE void CopyAverage( const double NumFrames )
 	{
 		_InclusiveTimeMS = _AvgInclusiveTimeMS;
+		_ExclusiveTimeMS = _AvgExclusiveTimeMS;
 		_NumCallsPerFrame = _AvgNumCallsPerFrame;
 
 		_FrameDurationMS /= NumFrames;
@@ -1185,6 +1189,7 @@ protected:
 				if (IsSelf())
 				{
 					Parent->_ExclusiveTimeMS = _InclusiveTimeMS;
+					Parent->_AvgExclusiveTimeMS = Parent->_ExclusiveTimeMS / NumFrames;
 					Parent->_ExclusiveTimePct = 100.0f * Parent->_ExclusiveTimeMS / Parent->_InclusiveTimeMS;
 				}
 
@@ -1201,6 +1206,7 @@ protected:
 		}
 
 		_AvgInclusiveTimeMS = _InclusiveTimeMS / NumFrames;
+		_AvgExclusiveTimeMS = _ExclusiveTimeMS / NumFrames;
 		_AvgNumCallsPerFrame = _NumCallsPerFrame / NumFrames;
 	}
 
@@ -1778,6 +1784,9 @@ public:
 	/** Average inclusive time of all instances for this event, in milliseconds. */
 	double _AvgInclusiveTimeMS;
 
+	/** Average exclusive time of all instances for this event, in milliseconds. */
+	double _AvgExclusiveTimeMS;
+
 	/** Number of times this event was called. */
 	double _NumCallsPerFrame;
 
