Disabled Profiler if STATS is not enabled (so Editor can compile also when ENABLE_STATNAMEDEVENTS is on).

diff --git a/Engine/Source/Developer/Profiler/Private/ProfilerCommands.cpp b/Engine/Source/Developer/Profiler/Private/ProfilerCommands.cpp
--- a/Engine/Source/Developer/Profiler/Private/ProfilerCommands.cpp
+++ b/Engine/Source/Developer/Profiler/Private/ProfilerCommands.cpp
@@ -1,6 +1,9 @@
 // Copyright Epic Games, Inc. All Rights Reserved.
 
 #include "ProfilerCommands.h"
+
+#if STATS
+
 #include "Misc/Paths.h"
 #include "Framework/MultiBox/MultiBoxBuilder.h"
 #include "EditorStyleSet.h"
@@ -354,3 +357,5 @@ bool FProfilerActionManager::OpenSettings_CanExecute() const
 }
 
 #undef LOCTEXT_NAMESPACE
+
+#endif // STATS
diff --git a/Engine/Source/Developer/Profiler/Private/ProfilerCommands.h b/Engine/Source/Developer/Profiler/Private/ProfilerCommands.h
--- a/Engine/Source/Developer/Profiler/Private/ProfilerCommands.h
+++ b/Engine/Source/Developer/Profiler/Private/ProfilerCommands.h
@@ -3,6 +3,9 @@
 #pragma once
 
 #include "CoreMinimal.h"
+
+#if STATS
+
 #include "Misc/Guid.h"
 #include "Framework/Commands/UIAction.h"
 #include "Framework/Commands/Commands.h"
@@ -222,3 +225,5 @@ protected:
 	/** Reference to the global instance of the profiler manager. */
 	class FProfilerManager* This;
 };
+
+#endif // STATS
diff --git a/Engine/Source/Developer/Profiler/Private/ProfilerDataProvider.cpp b/Engine/Source/Developer/Profiler/Private/ProfilerDataProvider.cpp
--- a/Engine/Source/Developer/Profiler/Private/ProfilerDataProvider.cpp
+++ b/Engine/Source/Developer/Profiler/Private/ProfilerDataProvider.cpp
@@ -2,6 +2,7 @@
 
 #include "ProfilerDataProvider.h"
 
+#if STATS
 
 /*-----------------------------------------------------------------------------
 	IDataProvider
@@ -167,3 +168,5 @@ const SIZE_T FArrayDataProvider::GetMemoryUsage() const
 	MemoryUsage += ChildrenIndicesMemoryUsage;
 	return MemoryUsage;
 }
+
+#endif // STATS
diff --git a/Engine/Source/Developer/Profiler/Private/ProfilerDataProvider.h b/Engine/Source/Developer/Profiler/Private/ProfilerDataProvider.h
--- a/Engine/Source/Developer/Profiler/Private/ProfilerDataProvider.h
+++ b/Engine/Source/Developer/Profiler/Private/ProfilerDataProvider.h
@@ -3,6 +3,9 @@
 #pragma once
 
 #include "CoreMinimal.h"
+
+#if STATS
+
 #include "ProfilerCommon.h"
 #include "ProfilerSample.h"
 
@@ -267,3 +270,5 @@ protected:
 
 	uint64 ChildrenIndicesMemoryUsage;
 };
+
+#endif // STATS
diff --git a/Engine/Source/Developer/Profiler/Private/ProfilerDataSource.cpp b/Engine/Source/Developer/Profiler/Private/ProfilerDataSource.cpp
--- a/Engine/Source/Developer/Profiler/Private/ProfilerDataSource.cpp
+++ b/Engine/Source/Developer/Profiler/Private/ProfilerDataSource.cpp
@@ -1,6 +1,9 @@
 // Copyright Epic Games, Inc. All Rights Reserved.
 
 #include "ProfilerDataSource.h"
+
+#if STATS
+
 #include "Containers/MapBuilder.h"
 #include "ProfilerStream.h"
 #include "ProfilerDataProvider.h"
@@ -566,3 +569,5 @@ FString EEventGraphTypes::ToDescription( const EEventGraphTypes::Type EventGraph
 }
 
 #undef LOCTEXT_NAMESPACE
+
+#endif // STATS
diff --git a/Engine/Source/Developer/Profiler/Private/ProfilerFPSAnalyzer.cpp b/Engine/Source/Developer/Profiler/Private/ProfilerFPSAnalyzer.cpp
--- a/Engine/Source/Developer/Profiler/Private/ProfilerFPSAnalyzer.cpp
+++ b/Engine/Source/Developer/Profiler/Private/ProfilerFPSAnalyzer.cpp
@@ -2,6 +2,8 @@
 
 #include "ProfilerFPSAnalyzer.h"
 
+#if STATS
+
 
 void FFPSAnalyzer::Reset()
 {
@@ -89,3 +91,5 @@ int32 FFPSAnalyzer::GetCount(float InMinVal, float InMaxVal)
 
 	return Histogram[Index].Count;
 }
+
+#endif // STATS
diff --git a/Engine/Source/Developer/Profiler/Private/ProfilerFPSAnalyzer.h b/Engine/Source/Developer/Profiler/Private/ProfilerFPSAnalyzer.h
--- a/Engine/Source/Developer/Profiler/Private/ProfilerFPSAnalyzer.h
+++ b/Engine/Source/Developer/Profiler/Private/ProfilerFPSAnalyzer.h
@@ -3,6 +3,9 @@
 #pragma once
 
 #include "CoreMinimal.h"
+
+#if STATS
+
 #include "ProfilerSample.h"
 
 
@@ -77,3 +80,5 @@ public:
 	int32 MaxVal;
 	int32 MinVal;
 };
+
+#endif // STATS
diff --git a/Engine/Source/Developer/Profiler/Private/ProfilerManager.cpp b/Engine/Source/Developer/Profiler/Private/ProfilerManager.cpp
--- a/Engine/Source/Developer/Profiler/Private/ProfilerManager.cpp
+++ b/Engine/Source/Developer/Profiler/Private/ProfilerManager.cpp
@@ -1,6 +1,9 @@
 // Copyright Epic Games, Inc. All Rights Reserved.
 
 #include "ProfilerManager.h"
+
+#if STATS
+
 #include "Modules/ModuleManager.h"
 #include "IProfilerClientModule.h"
 #include "Stats/StatsFile.h"
@@ -690,3 +693,5 @@ void FProfilerManager::SetViewMode( EProfilerViewMode NewViewMode )
 }
 
 #undef LOCTEXT_NAMESPACE
+
+#endif // STATS
diff --git a/Engine/Source/Developer/Profiler/Private/ProfilerManager.h b/Engine/Source/Developer/Profiler/Private/ProfilerManager.h
--- a/Engine/Source/Developer/Profiler/Private/ProfilerManager.h
+++ b/Engine/Source/Developer/Profiler/Private/ProfilerManager.h
@@ -3,6 +3,9 @@
 #pragma once
 
 #include "CoreMinimal.h"
+
+#if STATS
+
 #include "Misc/Guid.h"
 #include "Framework/Commands/UICommandList.h"
 #include "Misc/ConfigCacheIni.h"
@@ -519,3 +522,5 @@ protected:
 	/** True, if the profiler has at least one fully processed capture file. */
 	bool bHasCaptureFileFullyProcessed;
 };
+
+#endif // STATS
diff --git a/Engine/Source/Developer/Profiler/Private/ProfilerModule.cpp b/Engine/Source/Developer/Profiler/Private/ProfilerModule.cpp
--- a/Engine/Source/Developer/Profiler/Private/ProfilerModule.cpp
+++ b/Engine/Source/Developer/Profiler/Private/ProfilerModule.cpp
@@ -50,6 +50,7 @@ IMPLEMENT_MODULE( FProfilerModule, Profiler );
 
 TSharedRef<SWidget> FProfilerModule::CreateProfilerWindow( const TSharedRef<ISessionManager>& InSessionManager, const TSharedRef<SDockTab>& ConstructUnderMajorTab )
 {
+#if STATS
 	FProfilerManager::Initialize( InSessionManager );
 	TSharedRef<SProfilerWindow> ProfilerWindow = SNew( SProfilerWindow );
 	FProfilerManager::Get()->AssignProfilerWindow( ProfilerWindow );
@@ -57,20 +58,27 @@ TSharedRef<SWidget> FProfilerModule::CreateProfilerWindow( const TSharedRef<ISes
 	ConstructUnderMajorTab->SetOnTabClosed( SDockTab::FOnTabClosedCallback::CreateRaw( this, &FProfilerModule::Shutdown ) );
 
 	return ProfilerWindow;
+#else
+	return SNew(SBox);
+#endif // STATS
 }
 
 void FProfilerModule::ShutdownModule()
 {
+#if STATS
 	if (FProfilerManager::Get().IsValid())
 	{
 		FProfilerManager::Get()->Shutdown();
 	}
+#endif // STATS
 }
 
 void FProfilerModule::Shutdown( TSharedRef<SDockTab> TabBeingClosed )
 {
+#if STATS
 	FProfilerManager::Get()->Shutdown();
 	TabBeingClosed->SetOnTabClosed( SDockTab::FOnTabClosedCallback() );
+#endif // STATS
 }
 
 /*-----------------------------------------------------------------------------
@@ -79,6 +87,7 @@ void FProfilerModule::Shutdown( TSharedRef<SDockTab> TabBeingClosed )
 
 void FProfilerModule::StatsMemoryDumpCommand( const TCHAR* Filename )
 {
+#if STATS
 	TUniquePtr<FRawStatsMemoryProfiler> Instance( FStatsReader<FRawStatsMemoryProfiler>::Create( Filename ) );
 	if (Instance)
 	{
@@ -142,6 +151,7 @@ void FProfilerModule::StatsMemoryDumpCommand( const TCHAR* Filename )
 #endif // UE_BUILD_DEBUG
 		}
 	}
+#endif // STATS
 }
 
 
@@ -151,13 +161,13 @@ void FProfilerModule::StatsMemoryDumpCommand( const TCHAR* Filename )
 
 FRawStatsMemoryProfiler* FProfilerModule::OpenRawStatsForMemoryProfiling( const TCHAR* Filename )
 {
+#if STATS
 	FRawStatsMemoryProfiler* Instance = FStatsReader<FRawStatsMemoryProfiler>::Create( Filename );
 	if (Instance)
 	{
 		Instance->ReadAndProcessAsynchronously();
 	}
 	return Instance;
+#endif // STATS
+	return nullptr;
 }
-
-
-
diff --git a/Engine/Source/Developer/Profiler/Private/ProfilerRawStatsForMemory.cpp b/Engine/Source/Developer/Profiler/Private/ProfilerRawStatsForMemory.cpp
--- a/Engine/Source/Developer/Profiler/Private/ProfilerRawStatsForMemory.cpp
+++ b/Engine/Source/Developer/Profiler/Private/ProfilerRawStatsForMemory.cpp
@@ -1,6 +1,9 @@
 // Copyright Epic Games, Inc. All Rights Reserved.
 
 #include "ProfilerRawStatsForMemory.h"
+
+#if STATS
+
 #include "Stats/StatsMisc.h"
 #include "ProfilingDebugging/DiagnosticTable.h"
 
@@ -884,3 +887,5 @@ void FRawStatsMemoryProfiler::CompareSnapshotsHumanReadable( const FName BeginSn
 		out_Result.ValueSort( FCombinedAllocationInfoSizeGreater() );
 	}
 }
+
+#endif // STATS
diff --git a/Engine/Source/Developer/Profiler/Private/ProfilerRawStatsForMemory.h b/Engine/Source/Developer/Profiler/Private/ProfilerRawStatsForMemory.h
--- a/Engine/Source/Developer/Profiler/Private/ProfilerRawStatsForMemory.h
+++ b/Engine/Source/Developer/Profiler/Private/ProfilerRawStatsForMemory.h
@@ -3,6 +3,9 @@
 #pragma once
 
 #include "CoreMinimal.h"
+
+#if STATS
+
 #include "Stats/StatsFile.h"
 
 
@@ -383,3 +386,4 @@ protected:
 	int32 LastSequenceTagForNamedMarker;
 };
 
+#endif // STATS
diff --git a/Engine/Source/Developer/Profiler/Private/ProfilerRawStatsForThreadView.cpp b/Engine/Source/Developer/Profiler/Private/ProfilerRawStatsForThreadView.cpp
--- a/Engine/Source/Developer/Profiler/Private/ProfilerRawStatsForThreadView.cpp
+++ b/Engine/Source/Developer/Profiler/Private/ProfilerRawStatsForThreadView.cpp
@@ -1,6 +1,9 @@
 // Copyright Epic Games, Inc. All Rights Reserved.
 
 #include "ProfilerRawStatsForThreadView.h"
+
+#if STATS
+
 #include "HAL/FileManager.h"
 #include "Serialization/MemoryReader.h"
 #include "ProfilerDataProvider.h"
@@ -547,3 +550,5 @@ void FRawProfilerSession::ProcessStatPacketArray( const FStatPacketArray& StatPa
 
 	out_ProfilerFrame.SortChildren();
 }
+
+#endif // STATS
diff --git a/Engine/Source/Developer/Profiler/Private/ProfilerRawStatsForThreadView.h b/Engine/Source/Developer/Profiler/Private/ProfilerRawStatsForThreadView.h
--- a/Engine/Source/Developer/Profiler/Private/ProfilerRawStatsForThreadView.h
+++ b/Engine/Source/Developer/Profiler/Private/ProfilerRawStatsForThreadView.h
@@ -3,6 +3,9 @@
 #pragma once
 
 #include "CoreMinimal.h"
+
+#if STATS
+
 #include "Stats/StatsFile.h"
 #include "ProfilerStream.h"
 #include "ProfilerSession.h"
@@ -50,3 +53,5 @@ public:
 	 */
 	void ProcessStatPacketArray(const FStatPacketArray& PacketArray, FProfilerFrame& out_ProfilerFrame, int32 FrameIndex);
 };
+
+#endif // STATS
diff --git a/Engine/Source/Developer/Profiler/Private/ProfilerSample.cpp b/Engine/Source/Developer/Profiler/Private/ProfilerSample.cpp
--- a/Engine/Source/Developer/Profiler/Private/ProfilerSample.cpp
+++ b/Engine/Source/Developer/Profiler/Private/ProfilerSample.cpp
@@ -2,6 +2,7 @@
 
 #include "ProfilerSample.h"
 
+#if STATS
 
 #define LOCTEXT_NAMESPACE "ProfilerSample"
 
@@ -58,3 +59,5 @@ void FProfilerSample::FixupChildrenOrdering(const TMap<uint32,uint32>& ChildrenO
 
 
 #undef LOCTEXT_NAMESPACE
+
+#endif // STATS
diff --git a/Engine/Source/Developer/Profiler/Private/ProfilerSession.cpp b/Engine/Source/Developer/Profiler/Private/ProfilerSession.cpp
--- a/Engine/Source/Developer/Profiler/Private/ProfilerSession.cpp
+++ b/Engine/Source/Developer/Profiler/Private/ProfilerSession.cpp
@@ -1,6 +1,9 @@
 // Copyright Epic Games, Inc. All Rights Reserved.
 
 #include "ProfilerSession.h"
+
+#if STATS
+
 #include "ProfilerFPSAnalyzer.h"
 #include "ProfilerDataProvider.h"
 
@@ -985,3 +988,5 @@ void FProfilerAggregatedStat::Aggregate(const FProfilerSample& Sample, const TSh
 
 	_ValueOneFrame += TypedValue;
 }
+
+#endif // STATS
diff --git a/Engine/Source/Developer/Profiler/Private/ProfilerStream.cpp b/Engine/Source/Developer/Profiler/Private/ProfilerStream.cpp
--- a/Engine/Source/Developer/Profiler/Private/ProfilerStream.cpp
+++ b/Engine/Source/Developer/Profiler/Private/ProfilerStream.cpp
@@ -2,6 +2,7 @@
 
 #include "ProfilerStream.h"
 
+#if STATS
 
 /*-----------------------------------------------------------------------------
 	FProfilerUIStream
@@ -169,3 +170,5 @@ void FProfilerUIStream::CombineOrSet( FProfilerUIStackNode* ParentUIStackNode, c
 		}
 	}
 }
+
+#endif // STATS
diff --git a/Engine/Source/Developer/Profiler/Private/ProfilerStream.h b/Engine/Source/Developer/Profiler/Private/ProfilerStream.h
--- a/Engine/Source/Developer/Profiler/Private/ProfilerStream.h
+++ b/Engine/Source/Developer/Profiler/Private/ProfilerStream.h
@@ -3,6 +3,9 @@
 #pragma once
 
 #include "CoreMinimal.h"
+
+#if STATS
+
 #include "HAL/ThreadSafeCounter.h"
 #include "Containers/IndirectArray.h"
 #include "Containers/ChunkedArray.h"
@@ -563,3 +566,4 @@ protected:
 	}
 };
 
+#endif // STATS
diff --git a/Engine/Source/Developer/Profiler/Private/Widgets/SDataGraph.cpp b/Engine/Source/Developer/Profiler/Private/Widgets/SDataGraph.cpp
--- a/Engine/Source/Developer/Profiler/Private/Widgets/SDataGraph.cpp
+++ b/Engine/Source/Developer/Profiler/Private/Widgets/SDataGraph.cpp
@@ -1,6 +1,9 @@
 // Copyright Epic Games, Inc. All Rights Reserved.
 
 #include "Widgets/SDataGraph.h"
+
+#if STATS
+
 #include "Fonts/SlateFontInfo.h"
 #include "Misc/Paths.h"
 #include "Rendering/DrawElements.h"
@@ -1342,3 +1345,5 @@ void SDataGraph::EventGraph_OnRestoredFromHistory( uint32 FrameStartIndex, uint3
 }
 
 #undef LOCTEXT_NAMESPACE
+
+#endif // STATS
diff --git a/Engine/Source/Developer/Profiler/Private/Widgets/SDataGraph.h b/Engine/Source/Developer/Profiler/Private/Widgets/SDataGraph.h
--- a/Engine/Source/Developer/Profiler/Private/Widgets/SDataGraph.h
+++ b/Engine/Source/Developer/Profiler/Private/Widgets/SDataGraph.h
@@ -3,6 +3,9 @@
 #pragma once
 
 #include "CoreMinimal.h"
+
+#if STATS
+
 #include "Widgets/DeclarativeSyntaxSupport.h"
 #include "Layout/Geometry.h"
 #include "Input/CursorReply.h"
@@ -427,3 +430,5 @@ protected:
 
 	float FrameTimesMS[2];
 };
+
+#endif // STATS
diff --git a/Engine/Source/Developer/Profiler/Private/Widgets/SEventGraph.cpp b/Engine/Source/Developer/Profiler/Private/Widgets/SEventGraph.cpp
--- a/Engine/Source/Developer/Profiler/Private/Widgets/SEventGraph.cpp
+++ b/Engine/Source/Developer/Profiler/Private/Widgets/SEventGraph.cpp
@@ -1,6 +1,9 @@
 // Copyright Epic Games, Inc. All Rights Reserved.
 
 #include "Widgets/SEventGraph.h"
+
+#if STATS
+
 #include "Widgets/Layout/SSplitter.h"
 #include "Containers/MapBuilder.h"
 #include "Widgets/SOverlay.h"
@@ -3521,3 +3524,5 @@ void SEventGraph::FEventGraphState::CreateOneToOneMapping()
 }
 
 #undef LOCTEXT_NAMESPACE
+
+#endif // STATS
diff --git a/Engine/Source/Developer/Profiler/Private/Widgets/SEventGraph.h b/Engine/Source/Developer/Profiler/Private/Widgets/SEventGraph.h
--- a/Engine/Source/Developer/Profiler/Private/Widgets/SEventGraph.h
+++ b/Engine/Source/Developer/Profiler/Private/Widgets/SEventGraph.h
@@ -3,6 +3,9 @@
 #pragma once
 
 #include "CoreMinimal.h"
+
+#if STATS
+
 #include "SlateFwd.h"
 #include "Layout/Visibility.h"
 #include "Widgets/DeclarativeSyntaxSupport.h"
@@ -1278,3 +1281,5 @@ protected:
 	TSharedPtr<SComboBox<TSharedPtr<FName>>> ThreadFilterComboBox;
 	TArray<TSharedPtr<FName>> ThreadNamesForCombo;
 };
+
+#endif // STATS
diff --git a/Engine/Source/Developer/Profiler/Private/Widgets/SEventGraphTooltip.cpp b/Engine/Source/Developer/Profiler/Private/Widgets/SEventGraphTooltip.cpp
--- a/Engine/Source/Developer/Profiler/Private/Widgets/SEventGraphTooltip.cpp
+++ b/Engine/Source/Developer/Profiler/Private/Widgets/SEventGraphTooltip.cpp
@@ -1,6 +1,9 @@
 // Copyright Epic Games, Inc. All Rights Reserved.
 
 #include "Widgets/SEventGraphTooltip.h"
+
+#if STATS
+
 #include "Misc/Paths.h"
 #include "Fonts/SlateFontInfo.h"
 #include "Styling/CoreStyle.h"
@@ -514,3 +517,6 @@ EVisibility SEventGraphTooltip::GetHotPathIconVisibility(const TSharedPtr<FEvent
 
 
 #undef LOCTEXT_NAMESPACE
+
+#endif // STATS
+
diff --git a/Engine/Source/Developer/Profiler/Private/Widgets/SEventGraphTooltip.h b/Engine/Source/Developer/Profiler/Private/Widgets/SEventGraphTooltip.h
--- a/Engine/Source/Developer/Profiler/Private/Widgets/SEventGraphTooltip.h
+++ b/Engine/Source/Developer/Profiler/Private/Widgets/SEventGraphTooltip.h
@@ -3,6 +3,9 @@
 #pragma once
 
 #include "CoreMinimal.h"
+
+#if STATS
+
 #include "Layout/Visibility.h"
 #include "ProfilerDataSource.h"
 
@@ -21,3 +24,5 @@ protected:
 
 	static EVisibility GetHotPathIconVisibility(const TSharedPtr<FEventGraphSample> EventSample);
 };
+
+#endif // STATS
diff --git a/Engine/Source/Developer/Profiler/Private/Widgets/SFiltersAndPresets.cpp b/Engine/Source/Developer/Profiler/Private/Widgets/SFiltersAndPresets.cpp
--- a/Engine/Source/Developer/Profiler/Private/Widgets/SFiltersAndPresets.cpp
+++ b/Engine/Source/Developer/Profiler/Private/Widgets/SFiltersAndPresets.cpp
@@ -1,6 +1,9 @@
 // Copyright Epic Games, Inc. All Rights Reserved.
 
 #include "Widgets/SFiltersAndPresets.h"
+
+#if STATS
+
 #include "SlateOptMacros.h"
 #include "Widgets/Layout/SSeparator.h"
 #include "Widgets/Images/SImage.h"
@@ -1216,3 +1219,5 @@ FText SFiltersAndPresets::SortBy_GetSelectedText() const
 }
 
 #undef LOCTEXT_NAMESPACE
+
+#endif // STATS
diff --git a/Engine/Source/Developer/Profiler/Private/Widgets/SFiltersAndPresets.h b/Engine/Source/Developer/Profiler/Private/Widgets/SFiltersAndPresets.h
--- a/Engine/Source/Developer/Profiler/Private/Widgets/SFiltersAndPresets.h
+++ b/Engine/Source/Developer/Profiler/Private/Widgets/SFiltersAndPresets.h
@@ -3,6 +3,9 @@
 #pragma once
 
 #include "CoreMinimal.h"
+
+#if STATS
+
 #include "SlateFwd.h"
 #include "Widgets/DeclarativeSyntaxSupport.h"
 #include "ProfilerCommon.h"
@@ -424,3 +427,5 @@ protected:
 	/** Holds the visibility of each stat type. */
 	bool bStatTypeIsVisible[EProfilerSampleTypes::InvalidOrMax];
 };
+
+#endif // STATS
diff --git a/Engine/Source/Developer/Profiler/Private/Widgets/SHistogram.cpp b/Engine/Source/Developer/Profiler/Private/Widgets/SHistogram.cpp
--- a/Engine/Source/Developer/Profiler/Private/Widgets/SHistogram.cpp
+++ b/Engine/Source/Developer/Profiler/Private/Widgets/SHistogram.cpp
@@ -1,6 +1,9 @@
 // Copyright Epic Games, Inc. All Rights Reserved.
 
 #include "Widgets/SHistogram.h"
+
+#if STATS
+
 #include "Fonts/SlateFontInfo.h"
 #include "Styling/CoreStyle.h"
 #include "Misc/Paths.h"
@@ -172,4 +175,4 @@ void SHistogram::SetFPSAnalyzer(const TSharedPtr<FFPSAnalyzer>& InAnalyzer)
 
 #undef LOCTEXT_NAMESPACE
 
-
+#endif // STATS
diff --git a/Engine/Source/Developer/Profiler/Private/Widgets/SHistogram.h b/Engine/Source/Developer/Profiler/Private/Widgets/SHistogram.h
--- a/Engine/Source/Developer/Profiler/Private/Widgets/SHistogram.h
+++ b/Engine/Source/Developer/Profiler/Private/Widgets/SHistogram.h
@@ -3,6 +3,9 @@
 #pragma once
 
 #include "CoreMinimal.h"
+
+#if STATS
+
 #include "Widgets/DeclarativeSyntaxSupport.h"
 #include "Widgets/SCompoundWidget.h"
 #include "ProfilerSample.h"
@@ -120,3 +123,5 @@ public:
 protected:
 	FHistogramDescription Description;
 };
+
+#endif // STATS
diff --git a/Engine/Source/Developer/Profiler/Private/Widgets/SMultiDumpBrowser.cpp b/Engine/Source/Developer/Profiler/Private/Widgets/SMultiDumpBrowser.cpp
--- a/Engine/Source/Developer/Profiler/Private/Widgets/SMultiDumpBrowser.cpp
+++ b/Engine/Source/Developer/Profiler/Private/Widgets/SMultiDumpBrowser.cpp
@@ -2,6 +2,7 @@
 
 #include "Widgets/SMultiDumpBrowser.h"
 
+#if STATS
 
 #define LOCTEXT_NAMESPACE "SMultiDumpBrowser"
 
@@ -196,3 +197,5 @@ void SMultiDumpBrowser::PrefilterTextCommitted(const FText& InText, const ETextC
 
 
 #undef LOCTEXT_NAMESPACE
+
+#endif // STATS
diff --git a/Engine/Source/Developer/Profiler/Private/Widgets/SMultiDumpBrowser.h b/Engine/Source/Developer/Profiler/Private/Widgets/SMultiDumpBrowser.h
--- a/Engine/Source/Developer/Profiler/Private/Widgets/SMultiDumpBrowser.h
+++ b/Engine/Source/Developer/Profiler/Private/Widgets/SMultiDumpBrowser.h
@@ -3,6 +3,9 @@
 #pragma once
 
 #include "CoreMinimal.h"
+
+#if STATS
+
 #include "Widgets/DeclarativeSyntaxSupport.h"
 #include "ProfilerManager.h"
 #include "Widgets/SCompoundWidget.h"
@@ -136,3 +139,5 @@ protected:
 
 	FString TotalsFilteringText;
 };
+
+#endif // STATS
diff --git a/Engine/Source/Developer/Profiler/Private/Widgets/SProfilerFPSChartPanel.cpp b/Engine/Source/Developer/Profiler/Private/Widgets/SProfilerFPSChartPanel.cpp
--- a/Engine/Source/Developer/Profiler/Private/Widgets/SProfilerFPSChartPanel.cpp
+++ b/Engine/Source/Developer/Profiler/Private/Widgets/SProfilerFPSChartPanel.cpp
@@ -1,6 +1,9 @@
 // Copyright Epic Games, Inc. All Rights Reserved.
 
 #include "Widgets/SProfilerFPSChartPanel.h"
+
+#if STATS
+
 #include "Widgets/SBoxPanel.h"
 #include "SlateOptMacros.h"
 #include "Widgets/Layout/SBorder.h"
@@ -240,3 +243,5 @@ void SProfilerFPSChartPanel::ProfilerManager_OnViewModeChanged(EProfilerViewMode
 }
 
 #undef LOCTEXT_NAMESPACE
+
+#endif // STATS
diff --git a/Engine/Source/Developer/Profiler/Private/Widgets/SProfilerFPSChartPanel.h b/Engine/Source/Developer/Profiler/Private/Widgets/SProfilerFPSChartPanel.h
--- a/Engine/Source/Developer/Profiler/Private/Widgets/SProfilerFPSChartPanel.h
+++ b/Engine/Source/Developer/Profiler/Private/Widgets/SProfilerFPSChartPanel.h
@@ -3,6 +3,9 @@
 #pragma once
 
 #include "CoreMinimal.h"
+
+#if STATS
+
 #include "Widgets/DeclarativeSyntaxSupport.h"
 #include "Widgets/SCompoundWidget.h"
 #include "ProfilerManager.h"
@@ -48,3 +51,5 @@ protected:
 	/** The descriptions panel of the chart */
 	TSharedPtr<SProfilerFPSStatisticsPanel> StatisticsPanel;
 };
+
+#endif // STATS
diff --git a/Engine/Source/Developer/Profiler/Private/Widgets/SProfilerGraphPanel.cpp b/Engine/Source/Developer/Profiler/Private/Widgets/SProfilerGraphPanel.cpp
--- a/Engine/Source/Developer/Profiler/Private/Widgets/SProfilerGraphPanel.cpp
+++ b/Engine/Source/Developer/Profiler/Private/Widgets/SProfilerGraphPanel.cpp
@@ -1,6 +1,9 @@
 // Copyright Epic Games, Inc. All Rights Reserved.
 
 #include "Widgets/SProfilerGraphPanel.h"
+
+#if STATS
+
 #include "Widgets/SBoxPanel.h"
 #include "Widgets/Layout/SBorder.h"
 #include "Widgets/Layout/SScrollBar.h"
@@ -285,3 +288,5 @@ void SProfilerGraphPanel::SetScrollBarState()
 
 
 #undef LOCTEXT_NAMESPACE
+
+#endif // STATS
diff --git a/Engine/Source/Developer/Profiler/Private/Widgets/SProfilerGraphPanel.h b/Engine/Source/Developer/Profiler/Private/Widgets/SProfilerGraphPanel.h
--- a/Engine/Source/Developer/Profiler/Private/Widgets/SProfilerGraphPanel.h
+++ b/Engine/Source/Developer/Profiler/Private/Widgets/SProfilerGraphPanel.h
@@ -3,6 +3,9 @@
 #pragma once
 
 #include "CoreMinimal.h"
+
+#if STATS
+
 #include "Widgets/DeclarativeSyntaxSupport.h"
 #include "ProfilerManager.h"
 #include "Widgets/SCompoundWidget.h"
@@ -112,3 +115,5 @@ public:
 	/** Current view mode. */
 	EProfilerViewMode ViewMode;
 };
+
+#endif // STATS
diff --git a/Engine/Source/Developer/Profiler/Private/Widgets/SProfilerMiniView.cpp b/Engine/Source/Developer/Profiler/Private/Widgets/SProfilerMiniView.cpp
--- a/Engine/Source/Developer/Profiler/Private/Widgets/SProfilerMiniView.cpp
+++ b/Engine/Source/Developer/Profiler/Private/Widgets/SProfilerMiniView.cpp
@@ -1,6 +1,9 @@
 // Copyright Epic Games, Inc. All Rights Reserved.
 
 #include "Widgets/SProfilerMiniView.h"
+
+#if STATS
+
 #include "Fonts/SlateFontInfo.h"
 #include "Misc/Paths.h"
 #include "Rendering/DrawElements.h"
@@ -711,3 +714,5 @@ void SProfilerMiniView::ProcessData()
 
 	MaxFrameTime = FMath::Clamp( MaxFrameTime, 0.0f, (float)MAX_VISIBLE_THREADTIME );
 }
+
+#endif // STATS
diff --git a/Engine/Source/Developer/Profiler/Private/Widgets/SProfilerMiniView.h b/Engine/Source/Developer/Profiler/Private/Widgets/SProfilerMiniView.h
--- a/Engine/Source/Developer/Profiler/Private/Widgets/SProfilerMiniView.h
+++ b/Engine/Source/Developer/Profiler/Private/Widgets/SProfilerMiniView.h
@@ -3,6 +3,9 @@
 #pragma once
 
 #include "CoreMinimal.h"
+
+#if STATS
+
 #include "Widgets/DeclarativeSyntaxSupport.h"
 #include "ProfilerCommon.h"
 #include "Layout/Geometry.h"
@@ -414,3 +417,5 @@ private:
 	/** True if the active timer is currently registered */
 	bool bIsActiveTimerRegistered;
 };
+
+#endif // STATS
diff --git a/Engine/Source/Developer/Profiler/Private/Widgets/SProfilerSettings.cpp b/Engine/Source/Developer/Profiler/Private/Widgets/SProfilerSettings.cpp
--- a/Engine/Source/Developer/Profiler/Private/Widgets/SProfilerSettings.cpp
+++ b/Engine/Source/Developer/Profiler/Private/Widgets/SProfilerSettings.cpp
@@ -1,6 +1,9 @@
 // Copyright Epic Games, Inc. All Rights Reserved.
 
 #include "Widgets/SProfilerSettings.h"
+
+#if STATS
+
 #include "Fonts/SlateFontInfo.h"
 #include "Misc/Paths.h"
 #include "Widgets/SBoxPanel.h"
@@ -251,3 +254,6 @@ FReply SProfilerSettings::OptionDefault_OnClicked( bool* ValuePtr, const bool* D
 
 
 #undef LOCTEXT_NAMESPACE
+
+#endif // STATS
+
diff --git a/Engine/Source/Developer/Profiler/Private/Widgets/SProfilerSettings.h b/Engine/Source/Developer/Profiler/Private/Widgets/SProfilerSettings.h
--- a/Engine/Source/Developer/Profiler/Private/Widgets/SProfilerSettings.h
+++ b/Engine/Source/Developer/Profiler/Private/Widgets/SProfilerSettings.h
@@ -3,6 +3,9 @@
 #pragma once
 
 #include "CoreMinimal.h"
+
+#if STATS
+
 #include "Layout/Visibility.h"
 #include "Widgets/DeclarativeSyntaxSupport.h"
 #include "Input/Reply.h"
@@ -57,3 +60,5 @@ private:
 	/** Pointer to the profiler settings. */
 	FProfilerSettings* SettingPtr;
 };
+
+#endif // STATS
diff --git a/Engine/Source/Developer/Profiler/Private/Widgets/SProfilerThreadView.cpp b/Engine/Source/Developer/Profiler/Private/Widgets/SProfilerThreadView.cpp
--- a/Engine/Source/Developer/Profiler/Private/Widgets/SProfilerThreadView.cpp
+++ b/Engine/Source/Developer/Profiler/Private/Widgets/SProfilerThreadView.cpp
@@ -1,6 +1,9 @@
 // Copyright Epic Games, Inc. All Rights Reserved.
 
 #include "Widgets/SProfilerThreadView.h"
+
+#if STATS
+
 #include "Brushes/SlateColorBrush.h"
 #include "EditorStyleSet.h"
 
@@ -740,3 +743,5 @@ void SProfilerThreadView::UpdateInternalConstants()
 	NumPixelsPerMillisecond = (double)ThisGeometry.GetLocalSize().X / NumMillisecondsPerWindow;
 	NumMillisecondsPerSample = NumMillisecondsPerWindow / (double)ThisGeometry.GetLocalSize().X * (double)MIN_NUM_PIXELS_PER_SAMPLE;
 }
+
+#endif // STATS
diff --git a/Engine/Source/Developer/Profiler/Private/Widgets/SProfilerThreadView.h b/Engine/Source/Developer/Profiler/Private/Widgets/SProfilerThreadView.h
--- a/Engine/Source/Developer/Profiler/Private/Widgets/SProfilerThreadView.h
+++ b/Engine/Source/Developer/Profiler/Private/Widgets/SProfilerThreadView.h
@@ -3,6 +3,9 @@
 #pragma once
 
 #include "CoreMinimal.h"
+
+#if STATS
+
 #include "Misc/Paths.h"
 #include "Layout/SlateRect.h"
 #include "Layout/Geometry.h"
@@ -316,3 +319,5 @@ protected:
 	/** Cursor type. */
 	EThreadViewCursor::Type CursorType;
 };
+
+#endif // STATS
diff --git a/Engine/Source/Developer/Profiler/Private/Widgets/SProfilerToolbar.cpp b/Engine/Source/Developer/Profiler/Private/Widgets/SProfilerToolbar.cpp
--- a/Engine/Source/Developer/Profiler/Private/Widgets/SProfilerToolbar.cpp
+++ b/Engine/Source/Developer/Profiler/Private/Widgets/SProfilerToolbar.cpp
@@ -1,6 +1,9 @@
 // Copyright Epic Games, Inc. All Rights Reserved.
 
 #include "Widgets/SProfilerToolbar.h"
+
+#if STATS
+
 #include "Widgets/SBoxPanel.h"
 #include "Framework/Application/SlateApplication.h"
 #include "Widgets/Layout/SBorder.h"
@@ -185,3 +188,5 @@ void SProfilerToolbar::CreateCommands()
 
 
 #undef LOCTEXT_NAMESPACE
+
+#endif // STATS
diff --git a/Engine/Source/Developer/Profiler/Private/Widgets/SProfilerToolbar.h b/Engine/Source/Developer/Profiler/Private/Widgets/SProfilerToolbar.h
--- a/Engine/Source/Developer/Profiler/Private/Widgets/SProfilerToolbar.h
+++ b/Engine/Source/Developer/Profiler/Private/Widgets/SProfilerToolbar.h
@@ -1,9 +1,11 @@
 // Copyright Epic Games, Inc. All Rights Reserved.
 
-
 #pragma once
 
 #include "CoreMinimal.h"
+
+#if STATS
+
 #include "Widgets/DeclarativeSyntaxSupport.h"
 #include "Widgets/SCompoundWidget.h"
 
@@ -53,3 +55,5 @@ protected:
 protected:
 	TSharedPtr<SBorder> Border;
 };
+
+#endif // STATS
diff --git a/Engine/Source/Developer/Profiler/Private/Widgets/SProfilerWindow.cpp b/Engine/Source/Developer/Profiler/Private/Widgets/SProfilerWindow.cpp
--- a/Engine/Source/Developer/Profiler/Private/Widgets/SProfilerWindow.cpp
+++ b/Engine/Source/Developer/Profiler/Private/Widgets/SProfilerWindow.cpp
@@ -1,6 +1,9 @@
 // Copyright Epic Games, Inc. All Rights Reserved.
 
 #include "Widgets/SProfilerWindow.h"
+
+#if STATS
+
 #include "Widgets/SBoxPanel.h"
 #include "Widgets/Layout/SBorder.h"
 #include "Widgets/Text/STextBlock.h"
@@ -672,3 +675,5 @@ void SProfilerWindow::ProfilerManager_OnViewModeChanged( EProfilerViewMode NewVi
 
 
 #undef LOCTEXT_NAMESPACE
+
+#endif // STATS
diff --git a/Engine/Source/Developer/Profiler/Private/Widgets/SProfilerWindow.h b/Engine/Source/Developer/Profiler/Private/Widgets/SProfilerWindow.h
--- a/Engine/Source/Developer/Profiler/Private/Widgets/SProfilerWindow.h
+++ b/Engine/Source/Developer/Profiler/Private/Widgets/SProfilerWindow.h
@@ -3,6 +3,9 @@
 #pragma once
 
 #include "CoreMinimal.h"
+
+#if STATS
+
 #include "SlateFwd.h"
 #include "Misc/Guid.h"
 #include "Layout/Visibility.h"
@@ -168,3 +171,5 @@ public:
 	/** The number of seconds the profiler has been active */
 	float	DurationActive;
 };
+
+#endif // STATS
