Fixes for enum/int/float mismatches which warn under C++20 builds.

diff --git a/Engine/Plugins/Experimental/Animation/PoseSearch/Source/Runtime/Private/PoseSearchFeatureChannels.cpp b/Engine/Plugins/Experimental/Animation/PoseSearch/Source/Runtime/Private/PoseSearchFeatureChannels.cpp
--- a/Engine/Plugins/Experimental/Animation/PoseSearch/Source/Runtime/Private/PoseSearchFeatureChannels.cpp
+++ b/Engine/Plugins/Experimental/Animation/PoseSearch/Source/Runtime/Private/PoseSearchFeatureChannels.cpp
@@ -34,7 +34,7 @@ static FLinearColor GetColorForFeature(FPoseSearchFeatureDesc Feature, const FPo
 	const float FeatureIdx = Layout->Features.IndexOfByKey(Feature);
 	const float FeatureCountIdx = Layout->Features.Num() - 1;
 	const float FeatureCountIdxHalf = FeatureCountIdx / 2.f;
-	check(FeatureIdx != INDEX_NONE);
+	check(FeatureIdx != (float)INDEX_NONE);
 
 	const float Hue = FeatureIdx < FeatureCountIdxHalf
 		? FMath::GetMappedRangeValueUnclamped({ 0.f, FeatureCountIdxHalf }, FVector2f(60.f, 0.f), FeatureIdx)
diff --git a/Engine/Source/Developer/Profiler/Private/Widgets/SProfilerMiniView.cpp b/Engine/Source/Developer/Profiler/Private/Widgets/SProfilerMiniView.cpp
--- a/Engine/Source/Developer/Profiler/Private/Widgets/SProfilerMiniView.cpp
+++ b/Engine/Source/Developer/Profiler/Private/Widgets/SProfilerMiniView.cpp
@@ -495,8 +495,8 @@
 			const float StartEdgeDistance = FrameIndexToPosition(SelectionBoxFrameStart) - MousePosition.X;
 			const float EndEdgeDistance = MousePosition.X - FrameIndexToPosition(SelectionBoxFrameEnd);
 
-			bCanBeStartDragged = StartEdgeDistance < MOUSE_SNAP_DISTANCE && StartEdgeDistance > 0.0f;
-			bCanBeEndDragged = EndEdgeDistance < MOUSE_SNAP_DISTANCE && EndEdgeDistance > 0.0f;
+			bCanBeStartDragged = StartEdgeDistance < (float)MOUSE_SNAP_DISTANCE && StartEdgeDistance > 0.0f;
+			bCanBeEndDragged = EndEdgeDistance < (float)MOUSE_SNAP_DISTANCE && EndEdgeDistance > 0.0f;
 
 			if( StartEdgeDistance <= 0.0f && EndEdgeDistance <= 0.0f )
 			{
diff --git a/Engine/Source/Developer/Profiler/Private/Widgets/SProfilerThreadView.cpp b/Engine/Source/Developer/Profiler/Private/Widgets/SProfilerThreadView.cpp
--- a/Engine/Source/Developer/Profiler/Private/Widgets/SProfilerThreadView.cpp
+++ b/Engine/Source/Developer/Profiler/Private/Widgets/SProfilerThreadView.cpp
@@ -430,7 +430,7 @@
 		static const FSlateColorBrush SolidWhiteBrush = FSlateColorBrush( FColorList::White );
 		const FColor GameThreadColor = FColorList::Red;
 	
-		const FVector2D Position = FVector2D( UIStackNode.PositionXPx /*- PositionXMS*/, UIStackNode.PositionY*NUM_PIXELS_PER_ROW );
+		const FVector2D Position = FVector2D( UIStackNode.PositionXPx /*- PositionXMS*/, UIStackNode.PositionY*(double)NUM_PIXELS_PER_ROW );
 		const FVector2D Size = FVector2D( UIStackNode.WidthPx, NUM_PIXELS_PER_ROW );
 
 		// Draw a cycle counter for this profiler UI stack node.
@@ -607,7 +607,7 @@
 	const bool bZoomIn = MouseEvent.GetWheelDelta() < 0.0f;
 	const double Center = PositionXMS + RangeXMS*0.5f;
 
-	const double MinVisibleRangeMS = 1.0f / INV_MIN_VISIBLE_RANGE_X;
+	const double MinVisibleRangeMS = 1.0f / (double)INV_MIN_VISIBLE_RANGE_X;
 	const double NewUnclampedRange = bZoomIn ? RangeXMS*1.25f : RangeXMS / 1.25f;
 	const double NewRange = FMath::Clamp( NewUnclampedRange, MinVisibleRangeMS, FMath::Min( TotalRangeXMS, (double)MAX_VISIBLE_RANGE_X ) );
 
@@ -738,7 +738,7 @@
 	ZoomFactorX = (double)NUM_MILLISECONDS_PER_WINDOW / RangeXMS;
 	RangeY = FMath::RoundToFloat(ThisGeometry.GetLocalSize().Y / (double)NUM_PIXELS_PER_ROW);
 
-	const double Aspect = ThisGeometry.GetLocalSize().X / NUM_MILLISECONDS_PER_WINDOW * ZoomFactorX;
+	const double Aspect = ThisGeometry.GetLocalSize().X / (double)NUM_MILLISECONDS_PER_WINDOW * ZoomFactorX;
 	NumMillisecondsPerWindow = (double)ThisGeometry.GetLocalSize().X / Aspect;
 	NumPixelsPerMillisecond = (double)ThisGeometry.GetLocalSize().X / NumMillisecondsPerWindow;
 	NumMillisecondsPerSample = NumMillisecondsPerWindow / (double)ThisGeometry.GetLocalSize().X * (double)MIN_NUM_PIXELS_PER_SAMPLE;
diff --git a/Engine/Source/Editor/CurveTableEditor/Private/RealCurveModel.cpp b/Engine/Source/Editor/CurveTableEditor/Private/RealCurveModel.cpp
--- a/Engine/Source/Editor/CurveTableEditor/Private/RealCurveModel.cpp
+++ b/Engine/Source/Editor/CurveTableEditor/Private/RealCurveModel.cpp
@@ -42,7 +42,7 @@ void FRealCurveModel::DrawCurve(const FCurveEditor& CurveEditor, const FCurveEdi
 		double LastKeyValue =  double(RealCurve->GetKeyValue(LastKeyHandle));
 		double LastInterpMode = RealCurve->GetKeyInterpMode(LastKeyHandle);
 
-		if (LastInterpMode == RCIM_Constant)
+		if (LastInterpMode == (double)RCIM_Constant)
 		{
 			InterpolatingPoints.Add(MakeTuple( LastKeyTime, LastKeyValue));
 		}
@@ -52,7 +52,7 @@ void FRealCurveModel::DrawCurve(const FCurveEditor& CurveEditor, const FCurveEdi
 			auto KeyPair = RealCurve->GetKeyTimeValuePair(*It);
 
 			// if  constant , add another point to mark the end of the previous' key reign
-			if (LastInterpMode == RCIM_Constant)
+			if (LastInterpMode == (double)RCIM_Constant)
 			{
 				InterpolatingPoints.Add(MakeTuple( double(KeyPair.Key), LastKeyValue));
 			}
@@ -64,7 +64,7 @@ void FRealCurveModel::DrawCurve(const FCurveEditor& CurveEditor, const FCurveEdi
 			LastInterpMode = RealCurve->GetKeyInterpMode(LastKeyHandle);
 		}
 
-		if (LastInterpMode == RCIM_Constant)
+		if (LastInterpMode == (double)RCIM_Constant)
 		{
 			InterpolatingPoints.Add(MakeTuple( double(ScreenSpace.GetInputMax()), LastKeyValue));
 		}
diff --git a/Engine/Source/Editor/PluginWarden/Private/PluginWardenAuthorizer.cpp b/Engine/Source/Editor/PluginWarden/Private/PluginWardenAuthorizer.cpp
--- a/Engine/Source/Editor/PluginWarden/Private/PluginWardenAuthorizer.cpp
+++ b/Engine/Source/Editor/PluginWarden/Private/PluginWardenAuthorizer.cpp
@@ -289,7 +289,7 @@ EPluginAuthorizationState FPluginWardenAuthorizer::UpdateAuthorizationState(floa
 		case EPluginAuthorizationState::IsUserSignedIn_Waiting:
 		case EPluginAuthorizationState::SigninRequired_Waiting:
 		{
-			if ( WaitingTime > GeneralWaitingTimeout)
+			if ( WaitingTime > (float)GeneralWaitingTimeout)
 			{
 				NewState = EPluginAuthorizationState::Timeout;
 			}
@@ -297,7 +297,7 @@ EPluginAuthorizationState FPluginWardenAuthorizer::UpdateAuthorizationState(floa
 		}
 		case EPluginAuthorizationState::Signin_Waiting:
 		{
-			if ( WaitingTime > SigninWaitingTimeout )
+			if ( WaitingTime > (float)SigninWaitingTimeout )
 			{
 				NewState = EPluginAuthorizationState::SigninFailed;
 			}
diff --git a/Engine/Source/Runtime/Engine/Classes/Engine/NetSerialization.h b/Engine/Source/Runtime/Engine/Classes/Engine/NetSerialization.h
--- a/Engine/Source/Runtime/Engine/Classes/Engine/NetSerialization.h
+++ b/Engine/Source/Runtime/Engine/Classes/Engine/NetSerialization.h
@@ -304,7 +304,7 @@ bool WriteFixedCompressedFloat(const T Value, FArchive& Ar)
 	else
 	{
 		// We will scale up to get extra precision. But keep is a whole number preserve whole values
-		enum { Scale = Details::MaxBitValue / MaxValue };
+		constexpr int32 Scale = Details::MaxBitValue / MaxValue;
 		ScaledValue = FMath::RoundToInt( Scale * Value );
 	}
 
diff --git a/Engine/Source/Runtime/Engine/Private/SpeedTreeWind.cpp b/Engine/Source/Runtime/Engine/Private/SpeedTreeWind.cpp
--- a/Engine/Source/Runtime/Engine/Private/SpeedTreeWind.cpp
+++ b/Engine/Source/Runtime/Engine/Private/SpeedTreeWind.cpp
@@ -213,7 +213,7 @@ void FSpeedTreeWind::Advance(bool bEnabled, double fTime)
 		m_afRollingOffset[1] += m_afDirection[1] * m_fCombinedStrength * m_sParams.m_fRollingNoiseSpeed * m_fElapsedTime;
 
 		// compute oscillation indices
-		float fIndex = m_fCombinedStrength * (FSpeedTreeWind::NUM_WIND_POINTS_IN_CURVE - 1.0f);
+		float fIndex = m_fCombinedStrength * ((float)FSpeedTreeWind::NUM_WIND_POINTS_IN_CURVE - 1.0f);
 		int32 nBefore = static_cast<int>(fIndex);
 		int32 nAfter = nBefore + 1;
 		float fInterpolate;
diff --git a/Engine/Source/Runtime/Engine/Private/Streaming/TextureInstanceView.cpp b/Engine/Source/Runtime/Engine/Private/Streaming/TextureInstanceView.cpp
--- a/Engine/Source/Runtime/Engine/Private/Streaming/TextureInstanceView.cpp
+++ b/Engine/Source/Runtime/Engine/Private/Streaming/TextureInstanceView.cpp
@@ -526,7 +526,7 @@ void FRenderAssetInstanceAsyncView::ProcessElement(
 		else
 		{
 			check(AssetType == EStreamableRenderAssetType::StaticMesh || AssetType == EStreamableRenderAssetType::SkeletalMesh);
-			check(-TexelFactor <= MAX_MESH_LOD_COUNT);
+			check(-TexelFactor <= (float)MAX_MESH_LOD_COUNT);
 			MaxNumForcedLODs = FMath::Max(MaxNumForcedLODs, static_cast<int32>(-TexelFactor));
 		}
 
diff --git a/Engine/Source/Runtime/Engine/Private/UnrealClient.cpp b/Engine/Source/Runtime/Engine/Private/UnrealClient.cpp
--- a/Engine/Source/Runtime/Engine/Private/UnrealClient.cpp
+++ b/Engine/Source/Runtime/Engine/Private/UnrealClient.cpp
@@ -841,7 +841,7 @@ int32 FStatUnitData::DrawStat(FViewport* InViewport, FCanvas* InCanvas, int32 In
 		const float OutOfBudgetMarginHeight = (bSmallGraph ? 1 : 3);
 
 		const float GraphTotalWidth = GraphHorizPixelsPerFrame * NumberOfSamples;
-		const float GraphTotalHeight = TargetTimeMSHeight + (OutOfBudgetMarginHeight + EGS_UnboundedHighValueCount) * AlertPrintHeight;
+		const float GraphTotalHeight = TargetTimeMSHeight + (OutOfBudgetMarginHeight + (float)EGS_UnboundedHighValueCount) * AlertPrintHeight;
 
 		// Scale MS axis so that TargetTimeMS stays at fixed ordinate.
 		const float GraphVerticalPixelsPerMS = TargetTimeMSHeight / TargetTimeMS;
diff --git a/Engine/Source/Runtime/Engine/Private/WorldPartition/WorldPartitionStreamingPolicy.cpp b/Engine/Source/Runtime/Engine/Private/WorldPartition/WorldPartitionStreamingPolicy.cpp
--- a/Engine/Source/Runtime/Engine/Private/WorldPartition/WorldPartitionStreamingPolicy.cpp
+++ b/Engine/Source/Runtime/Engine/Private/WorldPartition/WorldPartitionStreamingPolicy.cpp
@@ -763,7 +763,7 @@ float FStreamingSourceVelocity::GetAverageVelocity(const FVector& NewPosition, c
 	VelocitiesHistory[LastIndex] = Velocity;
 
 	// return average
-	return (VelocitiesHistorySum / VELOCITY_HISTORY_SAMPLE_COUNT);
+	return (VelocitiesHistorySum / (float)VELOCITY_HISTORY_SAMPLE_COUNT);
 }
 
 #undef LOCTEXT_NAMESPACE
\ No newline at end of file
