Fixed unsafe type cast warnings in Profiler and ProfilerClient modules.

diff --git a/Engine/Source/Developer/Profiler/Private/ProfilerDataProvider.cpp b/Engine/Source/Developer/Profiler/Private/ProfilerDataProvider.cpp
--- a/Engine/Source/Developer/Profiler/Private/ProfilerDataProvider.cpp
+++ b/Engine/Source/Developer/Profiler/Private/ProfilerDataProvider.cpp
@@ -51,7 +51,7 @@ const uint32 IDataProvider::AdvanceFrame(const float DeltaTimeMS)
 	FrameIndices.Add(FIntPoint(FrameIndex, SampleEndIndex));
 
 	FrameTimes.Add(DeltaTimeMS);
-	ElapsedFrameTimes.Add(ElapsedTimeMS);
+	ElapsedFrameTimes.Add(static_cast<float>(ElapsedTimeMS));
 
 	// Update the values.
 	FrameCounters.Last() = LastSecondFrameCounter;
@@ -120,9 +120,9 @@ const uint32 FArrayDataProvider::AddHierarchicalSample(
 	if(FProfilerSample::IsIndexValid(InParentIndex))
 	{
 		FProfilerSample& Parent = Collection[ InParentIndex ];
-		const uint32 InitialMemUsage = Parent.ChildrenIndices().GetAllocatedSize();
+		const SIZE_T InitialMemUsage = Parent.ChildrenIndices().GetAllocatedSize();
 		Parent.AddChild(HierSampleIndex);
-		const uint32 AccumulatedMemoryUsage = Parent.ChildrenIndices().GetAllocatedSize() - InitialMemUsage;
+		const SIZE_T AccumulatedMemoryUsage = Parent.ChildrenIndices().GetAllocatedSize() - InitialMemUsage;
 		ChildrenIndicesMemoryUsage += AccumulatedMemoryUsage;
 	}
 
diff --git a/Engine/Source/Developer/Profiler/Private/ProfilerDataSource.cpp b/Engine/Source/Developer/Profiler/Private/ProfilerDataSource.cpp
--- a/Engine/Source/Developer/Profiler/Private/ProfilerDataSource.cpp
+++ b/Engine/Source/Developer/Profiler/Private/ProfilerDataSource.cpp
@@ -125,7 +125,7 @@ const uint32 FGraphDataSource::GetNumFrames() const
 
 const float FGraphDataSource::GetTotalTimeMS() const
 {
-	return ProfilerSession->GetDataProvider()->GetTotalTimeMS();
+	return static_cast<float>(ProfilerSession->GetDataProvider()->GetTotalTimeMS());
 }
 
 const TSharedRef<IDataProvider> FGraphDataSource::GetDataProvider() const
diff --git a/Engine/Source/Developer/Profiler/Private/ProfilerManager.cpp b/Engine/Source/Developer/Profiler/Private/ProfilerManager.cpp
--- a/Engine/Source/Developer/Profiler/Private/ProfilerManager.cpp
+++ b/Engine/Source/Developer/Profiler/Private/ProfilerManager.cpp
@@ -441,7 +441,7 @@ void FProfilerManager::ProfilerClient_OnProfilerFileTransfer( const FString& Fil
 {
 	// Display and update the notification a file that is being sent.
 
-	const float Progress = (double)FileProgress/(double)FileSize;
+	const float Progress = static_cast<float>((double)FileProgress/(double)FileSize);
 
 	ELoadingProgressStates ProgressState = ELoadingProgressStates::InvalidOrMax;
 
diff --git a/Engine/Source/Developer/Profiler/Private/ProfilerRawStatsForThreadView.cpp b/Engine/Source/Developer/Profiler/Private/ProfilerRawStatsForThreadView.cpp
--- a/Engine/Source/Developer/Profiler/Private/ProfilerRawStatsForThreadView.cpp
+++ b/Engine/Source/Developer/Profiler/Private/ProfilerRawStatsForThreadView.cpp
@@ -97,7 +97,7 @@ BreakPacketLoop:;
 	return Result;
 }
 
-static int64 GetFastThreadFrameTimeInternal( const FStatPacketArray& Frame, EThreadType::Type ThreadType )
+static uint32 GetFastThreadFrameTimeInternal( const FStatPacketArray& Frame, EThreadType::Type ThreadType )
 {
 	int64 Result = 0;
 
@@ -132,7 +132,8 @@ static int64 GetFastThreadFrameTimeInternal( const FStatPacketArray& Frame, EThr
 			}
 		}
 	}
-	return Result;
+
+	return static_cast<uint32>(Result);
 }
 
 void FRawProfilerSession::PrepareLoading()
@@ -265,7 +266,7 @@ void FRawProfilerSession::PrepareLoading()
 			const int64 TargetFrame = Frames[FrameIndex];
 			const FStatPacketArray& Frame = CombinedHistory.FindChecked( TargetFrame );
 
-			const double GameThreadTimeMS = GetMetaData()->ConvertCyclesToMS( GetFastThreadFrameTimeInternal( Frame, EThreadType::Game ) );
+			const double GameThreadTimeMS = GetMetaData()->ConvertCyclesToMS(GetFastThreadFrameTimeInternal( Frame, EThreadType::Game ));
 
 			if (GameThreadTimeMS == 0.0f)
 			{
@@ -293,19 +294,19 @@ void FRawProfilerSession::PrepareLoading()
 				const int64 TargetFrame = Frames[FrameIndex];
 				const FStatPacketArray& Frame = CombinedHistory.FindChecked(TargetFrame);
 
-				const double GameThreadTimeMS = GetMetaData()->ConvertCyclesToMS( GetFastThreadFrameTimeInternal(Frame,EThreadType::Game) );
+				const double GameThreadTimeMS = GetMetaData()->ConvertCyclesToMS(GetFastThreadFrameTimeInternal(Frame,EThreadType::Game));
 
 				if( GameThreadTimeMS == 0.0f )
 				{
 					continue;
 				}
 
-				const double RenderThreadTimeMS = GetMetaData()->ConvertCyclesToMS( GetFastThreadFrameTimeInternal(Frame,EThreadType::Renderer) );
+				const double RenderThreadTimeMS = GetMetaData()->ConvertCyclesToMS(GetFastThreadFrameTimeInternal(Frame,EThreadType::Renderer));
 
 				// Update mini-view, convert from cycles to ms.
 				TMap<uint32, float> ThreadTimesMS;
-				ThreadTimesMS.Add( GameThreadID, GameThreadTimeMS );
-				ThreadTimesMS.Add( GetMetaData()->GetRenderThreadID()[0], RenderThreadTimeMS );
+				ThreadTimesMS.Add( GameThreadID, static_cast<float>(GameThreadTimeMS) );
+				ThreadTimesMS.Add( GetMetaData()->GetRenderThreadID()[0], static_cast<float>(RenderThreadTimeMS) );
 
 				// Pass the reference to the stats' metadata.
 				OnAddThreadTime.ExecuteIfBound( FrameIndex, ThreadTimesMS, StatMetaData );
@@ -378,7 +379,7 @@ void FRawProfilerSession::ProcessStatPacketArray( const FStatPacketArray& StatPa
 	FProfilerSampleArray& MutableCollection = const_cast<FProfilerSampleArray&>(DataProvider->GetCollection());
 
 	// Add a root sample for this frame.
-	const uint32 FrameRootSampleIndex = DataProvider->AddHierarchicalSample( 0, MetaData->GetStatByID( 1 ).OwningGroup().ID(), 1, 0.0f, 0.0f, 1 );
+	const uint32 FrameRootSampleIndex = DataProvider->AddHierarchicalSample( 0, MetaData->GetStatByID( 1 ).OwningGroup().ID(), 1, 0, 0, 1 );
 
 	// Iterate through all stats packets and raw stats messages.
 	FName GameThreadFName = NAME_None;
@@ -482,8 +483,8 @@ void FRawProfilerSession::ProcessStatPacketArray( const FStatPacketArray& StatPa
 					const int64 Delta = int32( uint32( ScopeEnd.GetValue_int64() ) - uint32( ScopeStart.GetValue_int64() ) );
 					Current->CyclesEnd = Current->CyclesStart + Delta;
 
-					Current->CycleCounterStartTimeMS = MetaData->ConvertCyclesToMS( Current->CyclesStart );
-					Current->CycleCounterEndTimeMS = MetaData->ConvertCyclesToMS( Current->CyclesEnd );
+					Current->CycleCounterStartTimeMS = MetaData->ConvertCyclesToMS( uint32(Current->CyclesStart) );
+					Current->CycleCounterEndTimeMS = MetaData->ConvertCyclesToMS( uint32(Current->CyclesEnd) );
 
 					if (Current->CycleCounterStartTimeMS > Current->CycleCounterEndTimeMS)
 					{
@@ -495,7 +496,7 @@ void FRawProfilerSession::ProcessStatPacketArray( const FStatPacketArray& StatPa
 					FProfilerStackNode* ChildNode = Current;
 
 					// Update the child sample's DurationMS.
-					MutableCollection[ChildNode->SampleIndex].SetDurationCycles( Delta );
+					MutableCollection[ChildNode->SampleIndex].SetDurationCycles( uint32(Delta) );
 
 					verify( Current == Stack.Pop() );
 					Current = Stack.Last();
@@ -514,8 +515,8 @@ void FRawProfilerSession::ProcessStatPacketArray( const FStatPacketArray& StatPa
 			const int32 LastChildIndex = ThreadNode.Children.Num() - 1;
 			ThreadNode.CyclesStart = ThreadNode.Children[0]->CyclesStart;
 			ThreadNode.CyclesEnd = ThreadNode.Children[LastChildIndex]->CyclesEnd;
-			ThreadNode.CycleCounterStartTimeMS = MetaData->ConvertCyclesToMS( ThreadNode.CyclesStart );
-			ThreadNode.CycleCounterEndTimeMS = MetaData->ConvertCyclesToMS( ThreadNode.CyclesEnd );
+			ThreadNode.CycleCounterStartTimeMS = MetaData->ConvertCyclesToMS( uint32(ThreadNode.CyclesStart) );
+			ThreadNode.CycleCounterEndTimeMS = MetaData->ConvertCyclesToMS( uint32(ThreadNode.CyclesEnd) );
 
 			FProfilerSample& ProfilerSample = MutableCollection[ThreadNode.SampleIndex];
 			//ProfilerSample.SetStartAndEndMS( MetaData->ConvertCyclesToMS( ThreadNode.CyclesStart ), MetaData->ConvertCyclesToMS( ThreadNode.CyclesEnd ) );
@@ -525,13 +526,13 @@ void FRawProfilerSession::ProcessStatPacketArray( const FStatPacketArray& StatPa
 	// Get the game thread time.
 	check( GameThreadFName != NAME_None );
 	const FProfilerStackNode& GameThreadNode = *ThreadNodes.FindChecked( GameThreadFName );
-	const double GameThreadStartMS = MetaData->ConvertCyclesToMS( GameThreadNode.CyclesStart );
-	const double GameThreadEndMS = MetaData->ConvertCyclesToMS( GameThreadNode.CyclesEnd );
+	const double GameThreadStartMS = MetaData->ConvertCyclesToMS( uint32(GameThreadNode.CyclesStart) );
+	const double GameThreadEndMS = MetaData->ConvertCyclesToMS( uint32(GameThreadNode.CyclesEnd) );
 	//MutableCollection[FrameRootSampleIndex].SetStartAndEndMS( GameThreadStartMS, GameThreadEndMS );
 
 	// Advance frame
 	const uint32 LastFrameIndex = DataProvider->GetNumFrames();
-	DataProvider->AdvanceFrame( GameThreadEndMS - GameThreadStartMS );
+	DataProvider->AdvanceFrame( static_cast<float>(GameThreadEndMS - GameThreadStartMS) );
 
 	// Update aggregated stats
 	//UpdateAggregatedStats( LastFrameIndex );
diff --git a/Engine/Source/Developer/Profiler/Private/ProfilerSession.cpp b/Engine/Source/Developer/Profiler/Private/ProfilerSession.cpp
--- a/Engine/Source/Developer/Profiler/Private/ProfilerSession.cpp
+++ b/Engine/Source/Developer/Profiler/Private/ProfilerSession.cpp
@@ -460,7 +460,7 @@ bool FProfilerSession::HandleTicker( float DeltaTime )
 					ThreadDuration.DurationCycles, 1, 
 					FrameRootSampleIndex 
 				);
-				ThreadMS.FindOrAdd( ThreadID ) = StatMetaData->ConvertCyclesToMS( ThreadDuration.DurationCycles );
+				ThreadMS.FindOrAdd( ThreadID ) = static_cast<float>(StatMetaData->ConvertCyclesToMS( ThreadDuration.DurationCycles ));
 
 				// Recursively add children and parent to the root samples.
 				for( int32 Index = 0; Index < ThreadGraph.Children.Num(); Index++ )
@@ -486,7 +486,7 @@ bool FProfilerSession::HandleTicker( float DeltaTime )
 		MutableCollection[FrameRootSampleIndex].SetDurationCycles( GameThreadCycles != 0 ? GameThreadCycles : MaxThreadCycles );
 
 		// Update FPS analyzer.
-		const float GameThreadTimeMS = StatMetaData->ConvertCyclesToMS( GameThreadCycles );
+		const float GameThreadTimeMS = static_cast<float>(StatMetaData->ConvertCyclesToMS( GameThreadCycles ));
 		FPSAnalyzer->AddSample( GameThreadTimeMS > 0.0f ? 1000.0f / GameThreadTimeMS : 0.0f );
 
 		// Process the non-hierarchical samples for the specified frame.
@@ -509,7 +509,7 @@ bool FProfilerSession::HandleTicker( float DeltaTime )
 
 		// Advance frame
 		const uint32 DataProviderFrameIndex = DataProvider->GetNumFrames();
-		DataProvider->AdvanceFrame( StatMetaData->ConvertCyclesToMS( MaxThreadCycles ) );
+		DataProvider->AdvanceFrame( static_cast<float>(StatMetaData->ConvertCyclesToMS( MaxThreadCycles )) );
 
 		// Update aggregated stats
 		UpdateAggregatedStats( DataProviderFrameIndex );
@@ -948,17 +948,17 @@ void FProfilerAggregatedStat::Advance()
 		}
 
 		{
-			_MinValueAllFrames = FMath::Min<float>( _MinValueAllFrames, _ValueOneFrame );
-			_MaxValueAllFrames = FMath::Max<float>( _MaxValueAllFrames, _ValueOneFrame );
+			_MinValueAllFrames = FMath::Min<double>( _MinValueAllFrames, _ValueOneFrame );
+			_MaxValueAllFrames = FMath::Max<double>( _MaxValueAllFrames, _ValueOneFrame );
 
-			_MinNumCallsAllFrames = FMath::Min<float>( _MinNumCallsAllFrames, _NumCallsOneFrame );
-			_MaxNumCallsAllFrames = FMath::Max<float>( _MaxNumCallsAllFrames, _NumCallsOneFrame );
+			_MinNumCallsAllFrames = FMath::Min<uint32>( _MinNumCallsAllFrames, _NumCallsOneFrame );
+			_MaxNumCallsAllFrames = FMath::Max<uint32>( _MaxNumCallsAllFrames, _NumCallsOneFrame );
 		}
 	}
 	else
 	{
-		_MinValueAllFrames = FMath::Min<float>( _MinValueAllFrames, _ValueOneFrame );
-		_MaxValueAllFrames = FMath::Max<float>( _MaxValueAllFrames, _ValueOneFrame );
+		_MinValueAllFrames = FMath::Min<double>( _MinValueAllFrames, _ValueOneFrame );
+		_MaxValueAllFrames = FMath::Max<double>( _MaxValueAllFrames, _ValueOneFrame );
 	}
 
 	_ValueOneFrame = 0.0f;
diff --git a/Engine/Source/Developer/Profiler/Private/ProfilerSession.h b/Engine/Source/Developer/Profiler/Private/ProfilerSession.h
--- a/Engine/Source/Developer/Profiler/Private/ProfilerSession.h
+++ b/Engine/Source/Developer/Profiler/Private/ProfilerSession.h
@@ -688,7 +688,7 @@ public:
 	/** @return The average number of calls of all combined instances. */
 	const float AvgNumCalls() const 
 	{
-		return (double)_NumCallsAllFrames / (double)_NumFrames;
+		return static_cast<float>((double)_NumCallsAllFrames / (double)_NumFrames);
 	}
 
 	/** @return The min number of calls of all combined instances. */
@@ -708,7 +708,7 @@ public:
 	 */
 	const float FramesWithCallPct() const
 	{
-		return (double)_NumFramesWithCall / (double)_NumFrames * 100.0f;
+		return static_cast<float>((double)_NumFramesWithCall / (double)_NumFrames * 100.0);
 	}
 
 	/**
diff --git a/Engine/Source/Developer/Profiler/Private/ProfilerStream.h b/Engine/Source/Developer/Profiler/Private/ProfilerStream.h
--- a/Engine/Source/Developer/Profiler/Private/ProfilerStream.h
+++ b/Engine/Source/Developer/Profiler/Private/ProfilerStream.h
@@ -27,7 +27,7 @@ struct FInclusiveTime
 	/** Number of recursion. */
 	int32 Recursion;
 	FInclusiveTime()
-		: DurationCycles( 0.0 )
+		: DurationCycles( 0 )
 		, CallCount( 0 )
 		, Recursion( 0 )
 	{}
diff --git a/Engine/Source/Developer/Profiler/Private/Widgets/SDataGraph.cpp b/Engine/Source/Developer/Profiler/Private/Widgets/SDataGraph.cpp
--- a/Engine/Source/Developer/Profiler/Private/Widgets/SDataGraph.cpp
+++ b/Engine/Source/Developer/Profiler/Private/Widgets/SDataGraph.cpp
@@ -270,7 +270,7 @@ void SDataGraph::UpdateState()
 			NumDataPoints = (int32)FirstTrackedStat->GraphDataSource->GetNumFrames();
 		}
 
-		NumVisiblePoints = FMath::Max( 0, FMath::TruncToInt(ThisGeometry.Size.X) / DistanceBetweenPoints );
+		NumVisiblePoints = FMath::Max( 0, static_cast<int32>(FMath::TruncToInt(ThisGeometry.Size.X)) / DistanceBetweenPoints );
 		// GraphOffset - Updated by OnMouseMove or by ScrollTo
 		GraphOffset = FMath::Clamp( GraphOffset, 0, FMath::Max(NumDataPoints-NumVisiblePoints,0) );
 		
@@ -306,7 +306,7 @@ int32 SDataGraph::OnPaint( const FPaintArgs& Args, const FGeometry& AllottedGeom
 
 	/** Width of the alloted geometry that is used to draw a data graph. */
 	const float AreaX0 = 0.0f;
-	const float AreaX1 = AllottedGeometry.GetLocalSize().X;
+	const float AreaX1 = static_cast<float>(AllottedGeometry.GetLocalSize().X);
 
 	// Draw background.
 	FSlateDrawElement::MakeBox
@@ -337,7 +337,7 @@ int32 SDataGraph::OnPaint( const FPaintArgs& Args, const FGeometry& AllottedGeom
 		SCOPE_CYCLE_COUNTER(STAT_DG_OnPaint);
 
 		const FTrackedStat& TrackedStat = *It.Value();
-		const float GraphYScale = AllottedGeometry.GetLocalSize().Y/ScaleY;
+		const float GraphYScale = static_cast<float>(AllottedGeometry.GetLocalSize().Y) / ScaleY;
 		
 		const float UnitTypeScale = TrackedStat.GraphDataSource->GetSampleType() != EProfilerSampleTypes::Memory ? 1.0f : CounterToTimeScale;
 		const float TimeAccuracyMS = FTimeAccuracy::AsFrameTime( TimeBasedAccuracy );
@@ -427,8 +427,8 @@ int32 SDataGraph::OnPaint( const FPaintArgs& Args, const FGeometry& AllottedGeom
 				for (float GraphStartTimeMS = GraphOffsetMS; GraphStartTimeMS < GraphRangeEndMS; GraphStartTimeMS += TimeAccuracyMS)
 				{
 					const float Value = GraphDataSource->GetValueFromTimeRange( GraphStartTimeMS, GraphStartTimeMS + TimeAccuracyMS );
-					const float XPos = DistanceBetweenPoints*GraphPoints.Num();
-					const float YPos = FMath::Clamp( AllottedGeometry.GetLocalSize().Y - GraphYScale*Value*UnitTypeScale, 0.0f, AllottedGeometry.GetLocalSize().Y );
+					const float XPos = static_cast<float>(DistanceBetweenPoints * GraphPoints.Num());
+					const float YPos = FMath::Clamp( static_cast<float>(AllottedGeometry.GetLocalSize().Y) - GraphYScale*Value*UnitTypeScale, 0.0f, static_cast<float>(AllottedGeometry.GetLocalSize().Y) );
 					new (GraphPoints)FVector2D( XPos, YPos );
 				}
 
@@ -458,7 +458,7 @@ int32 SDataGraph::OnPaint( const FPaintArgs& Args, const FGeometry& AllottedGeom
 				{
 					const float Value = GraphDataSource->GetValueFromIndex( GraphStartIndex );
 					const float XPos = DistanceBetweenPoints*(float)GraphPoints.Num();
-					const float YPos = FMath::Clamp( AllottedGeometry.Size.Y - GraphYScale*Value*UnitTypeScale, 0.0f, AllottedGeometry.Size.Y );
+					const float YPos = FMath::Clamp( static_cast<float>(AllottedGeometry.Size.Y) - GraphYScale*Value*UnitTypeScale, 0.0f, static_cast<float>(AllottedGeometry.Size.Y) );
 					GraphPoints.Add( FVector2D(XPos,YPos) );
 				}
  
@@ -481,7 +481,7 @@ int32 SDataGraph::OnPaint( const FPaintArgs& Args, const FGeometry& AllottedGeom
 
 
 	FSlateFontInfo SummaryFont = FCoreStyle::GetDefaultFontStyle("Regular", 8);
-	const float MaxFontCharHeight = FontMeasureService->Measure( TEXT("!"), SummaryFont ).Y;
+	const float MaxFontCharHeight = static_cast<float>(FontMeasureService->Measure( TEXT("!"), SummaryFont ).Y);
 
 	// Draw graph annotations.
 	//const IDataProviderPtr DataProvider = StatIDToGraphDescriptionMapping.Num() > 0 ? StatIDToGraphDescriptionMapping.CreateConstIterator().Key().DataSource->GetDataProvider() : NULL;
@@ -514,7 +514,7 @@ int32 SDataGraph::OnPaint( const FPaintArgs& Args, const FGeometry& AllottedGeom
 
 			for( int32 FrameIndex = FrameStartIndex; FrameIndex < FrameEndIndex; FrameIndex += AvgFrameRate )
 			{
-				const float MarkerPosX = (FrameIndex - GraphOffset) * DistanceBetweenPoints;
+				const float MarkerPosX = static_cast<float>((FrameIndex - GraphOffset) * DistanceBetweenPoints);
 				const float ElapsedFrameTimeMS = DataProvider->GetElapsedFrameTimeMS( FrameIndex );
 
 				LinePoints.Add( FVector2D(MarkerPosX, 0.0) );
@@ -574,7 +574,7 @@ int32 SDataGraph::OnPaint( const FPaintArgs& Args, const FGeometry& AllottedGeom
 
 			for( int32 FrameIndex = FrameStartIndex; FrameIndex < FrameEndIndex; FrameIndex += AvgFrameRate )
 			{
-				const float MarkerPosX = (FrameIndex - GraphOffset) * DistanceBetweenPoints;
+				const float MarkerPosX = static_cast<float>((FrameIndex - GraphOffset) * DistanceBetweenPoints);
 				const float ElapsedFrameTimeMS = FrameIndex * FTimeAccuracy::AsFrameTime( TimeBasedAccuracy );
 				const int32 ElapsedFrameTime = FMath::Max( FMath::RoundToInt( ElapsedFrameTimeMS * 0.001f ) - 1, 0 );
 				const int32 AccumulatedFrameCounter = bCanBeDisplayedAsMulti ? FrameIndex : DataProvider->GetAccumulatedFrameCounter(ElapsedFrameTime);
@@ -631,7 +631,7 @@ int32 SDataGraph::OnPaint( const FPaintArgs& Args, const FGeometry& AllottedGeom
 
 	//-----------------------------------------------------------------------------
 
-	const int32 MaxGridPixelSpacing = 160.0f;
+	const int32 MaxGridPixelSpacing = 160;
 
 	// Draw a horizontal lines every 150 pixels and draw a few basic lines
 	static const TArray<float> DefaultTimeValueHints = TArrayBuilder<float>()
@@ -658,7 +658,7 @@ int32 SDataGraph::OnPaint( const FPaintArgs& Args, const FGeometry& AllottedGeom
 	const float MinTimeValue = 0.0f;
 	const float MaxTimeValue = ScaleY;
 	const float TimeValueGraphScale = MaxTimeValue / SecondaryIndicators;
-	const float TimeValueToGraph = AllottedGeometry.GetLocalSize().Y/MaxTimeValue;
+	const float TimeValueToGraph = static_cast<float>(AllottedGeometry.GetLocalSize().Y) / MaxTimeValue;
 
 	for( int32 SecondaryIndex = 1; SecondaryIndex <= SecondaryIndicators; SecondaryIndex++ )
 	{
@@ -755,7 +755,7 @@ int32 SDataGraph::OnPaint( const FPaintArgs& Args, const FGeometry& AllottedGeom
 			continue;
 		}
 
-		const float MarkerPosY = AllottedGeometry.GetLocalSize().Y - TimeValue*TimeValueToGraph;
+		const float MarkerPosY = static_cast<float>(AllottedGeometry.GetLocalSize().Y) - TimeValue * TimeValueToGraph;
 
 		// Check if this hint should be drawn as the basic hint.
 		const FLinearColor* BasicHintColor = DefaultTimeValueHintColors.Find( TimeValue );
@@ -807,7 +807,7 @@ int32 SDataGraph::OnPaint( const FPaintArgs& Args, const FGeometry& AllottedGeom
 
 		// Right	-	Values in human readable string, for the non-hierarchical samples
 		const FString CounterValueStr = FString::Printf( TEXT("%.1f KB"), TimeValue/CounterToTimeScale );
-		const float RightValueSizeX = FontMeasureService->Measure( CounterValueStr, SummaryFont ).X;
+		const float RightValueSizeX = static_cast<float>(FontMeasureService->Measure( CounterValueStr, SummaryFont ).X);
 		FSlateDrawElement::MakeText
 		(
 			OutDrawElements, 
@@ -823,7 +823,7 @@ int32 SDataGraph::OnPaint( const FPaintArgs& Args, const FGeometry& AllottedGeom
 	// Draw selected frames markers.
 	{
 		LayerId++;
-		const float LocalGraphOffset = GraphOffset * DistanceBetweenPoints;
+		const float LocalGraphOffset = static_cast<float>(GraphOffset * DistanceBetweenPoints);
 		const float LocalGraphSelectionX0 = FrameIndices[0]*DistanceBetweenPoints - LocalGraphOffset;
 		const float LocalGraphSelectionX1 = FrameIndices[1]*DistanceBetweenPoints - LocalGraphOffset;
 		const float LocalGraphSelectionX[2] = { LocalGraphSelectionX0, LocalGraphSelectionX1 };
@@ -874,7 +874,7 @@ int32 SDataGraph::OnPaint( const FPaintArgs& Args, const FGeometry& AllottedGeom
 		LayerId++;
 
 		const int32 LocalPosition = HoveredFrameIndex - GraphOffset;
-		const float LocalPositionGraphX = LocalPosition * DistanceBetweenPoints;
+		const float LocalPositionGraphX = static_cast<float>(LocalPosition * DistanceBetweenPoints);
 
 		FSlateDrawElement::MakeBox
 		(
@@ -1003,7 +1003,7 @@ FReply SDataGraph::OnMouseButtonDown( const FGeometry& MyGeometry, const FPointe
 		if( NumDataPoints > 0 )
 		{
 			// Capture mouse, so we can scroll outside this widget.
-			RealGraphOffset = GraphOffset;
+			RealGraphOffset = static_cast<float>(GraphOffset);
 			Reply = FReply::Handled().CaptureMouse(SharedThis(this));
 		}
 	}
@@ -1095,7 +1095,7 @@ FReply SDataGraph::OnMouseMove( const FGeometry& MyGeometry, const FPointerEvent
 		if( HasMouseCapture() && !MouseEvent.GetCursorDelta().IsZero() )
 		{
 			bIsRMB_Scrolling = true;
-			const float ScrollByAmount = -MouseEvent.GetCursorDelta().X * (1.0f/DistanceBetweenPoints);
+			const float ScrollByAmount = -static_cast<float>(MouseEvent.GetCursorDelta().X) * (1.0f/DistanceBetweenPoints);
 			RealGraphOffset += ScrollByAmount;
 
 			GraphOffset = FMath::Clamp( FMath::TruncToInt( RealGraphOffset ), 0, FMath::Max(NumDataPoints-NumVisiblePoints,0) );
@@ -1111,7 +1111,7 @@ FReply SDataGraph::OnMouseMove( const FGeometry& MyGeometry, const FPointerEvent
 const int32 SDataGraph::CalculateFrameIndex( const FVector2D InMousePosition ) const
 {
 	const float ScaleX = 1.0f/DistanceBetweenPoints;
-	const int32 MousePositionOffset = FMath::TruncToInt( (InMousePosition.X+HalfGraphMarkerWidth) * ScaleX );
+	const int32 MousePositionOffset = FMath::TruncToInt( (static_cast<float>(InMousePosition.X) + HalfGraphMarkerWidth) * ScaleX );
 	return FMath::Clamp( GraphOffset+MousePositionOffset, 0, NumDataPoints-1 );
 }
 
diff --git a/Engine/Source/Developer/Profiler/Private/Widgets/SEventGraph.cpp b/Engine/Source/Developer/Profiler/Private/Widgets/SEventGraph.cpp
--- a/Engine/Source/Developer/Profiler/Private/Widgets/SEventGraph.cpp
+++ b/Engine/Source/Developer/Profiler/Private/Widgets/SEventGraph.cpp
@@ -528,7 +528,7 @@ protected:
 	{
 		const FLinearColor ThreadColor(5.0f,0.0f,0.0f,1.0f);
 		const FLinearColor DefaultColor(0.0f,0.0f,0.0f,0.0f);
-		const float Alpha = EventPtr->_FramePct * 0.01f;
+		const float Alpha = static_cast<float>(EventPtr->_FramePct) * 0.01f;
 		const FLinearColor BackgroundColorAndOpacity = FMath::Lerp(DefaultColor,ThreadColor,Alpha);
 		return BackgroundColorAndOpacity;
 	}
@@ -2814,8 +2814,8 @@ void SEventGraph::GenerateTopEvents( const TSet< FEventGraphSamplePtr >& EventPt
 	for( int32 Nx = 0; Nx < EventPtrArray.Num() && Nx < NumTopEvents; ++Nx )
 	{
 		const FEventGraphSamplePtr EventPtr = EventPtrArray[Nx];
-		const float IncTimeToTotalPct = EventPtr->_InclusiveTimeMS / TotalTimeMS;
-		const float HeightPct = EventPtr->_InclusiveTimeMS / Top5TimeMS;
+		const float IncTimeToTotalPct = static_cast<float>(EventPtr->_InclusiveTimeMS / TotalTimeMS);
+		const float HeightPct = static_cast<float>(EventPtr->_InclusiveTimeMS / Top5TimeMS);
 		out_Results.Add( FEventPtrAndMisc(EventPtr,IncTimeToTotalPct,HeightPct) );
 	}
 }
diff --git a/Engine/Source/Developer/Profiler/Private/Widgets/SEventGraphTooltip.cpp b/Engine/Source/Developer/Profiler/Private/Widgets/SEventGraphTooltip.cpp
--- a/Engine/Source/Developer/Profiler/Private/Widgets/SEventGraphTooltip.cpp
+++ b/Engine/Source/Developer/Profiler/Private/Widgets/SEventGraphTooltip.cpp
@@ -30,7 +30,7 @@ TSharedPtr<SToolTip> SEventGraphTooltip::GetTableCellTooltip( const TSharedPtr<F
 
 	const FLinearColor ThreadColor(5.0f,0.0f,0.0f,1.0f);
 	const FLinearColor DefaultColor(1.0f,1.0f,1.0f,1.0f);
-	const float Alpha = EventSample->_FramePct * 0.01f;
+	const float Alpha = static_cast<float>(EventSample->_FramePct) * 0.01f;
 	const FLinearColor ColorAndOpacity = FMath::Lerp(DefaultColor,ThreadColor,Alpha);
 
 	const FText InclusiveTimePctCaller = FText::Format(LOCTEXT("PctOfTheCaller", "({0} of the caller)"), FText::FromString(EventSample->GetFormattedValue(EEventPropertyIndex::InclusiveTimePct)));
@@ -463,7 +463,8 @@ TSharedPtr<SToolTip> SEventGraphTooltip::GetTableCellTooltip( const TSharedPtr<F
 		for( int32 ChildIndex = 0; ChildIndex < Children.Num(); ChildIndex++ )
 		{
 			const FEventGraphSamplePtr& Child = Children[ChildIndex];
-			MinimalChildren.Add( FEventNameAndPct(Child->_InclusiveTimePct,Child->GetShortEventName()) );
+			const float ChildInclusiveTimePct = static_cast<float>(Child->_InclusiveTimePct);
+			MinimalChildren.Add( FEventNameAndPct(ChildInclusiveTimePct, Child->GetShortEventName()) );
 		}
 
 		struct FCompareByFloatDescending
diff --git a/Engine/Source/Developer/Profiler/Private/Widgets/SHistogram.cpp b/Engine/Source/Developer/Profiler/Private/Widgets/SHistogram.cpp
--- a/Engine/Source/Developer/Profiler/Private/Widgets/SHistogram.cpp
+++ b/Engine/Source/Developer/Profiler/Private/Widgets/SHistogram.cpp
@@ -65,7 +65,7 @@ int32 SHistogram::OnPaint( const FPaintArgs& Args, const FGeometry& AllottedGeom
 	static const FLinearColor GridTextColor = FLinearColor(1.0f,1.0f,1.0f, 0.25f);
 	static const FLinearColor BorderColor = FLinearColor(0.0f,0.0f,0.0f,1.0f);
 	FSlateFontInfo SummaryFont = FCoreStyle::GetDefaultFontStyle("Regular", 8);
-	const float MaxFontCharHeight = FontMeasureService->Measure( TEXT("!"), SummaryFont ).Y;
+	const float MaxFontCharHeight = static_cast<float>(FontMeasureService->Measure( TEXT("!"), SummaryFont ).Y);
 	TArray<FVector2D> LinePoints;
 
 	// draw the histogram box
@@ -103,7 +103,7 @@ int32 SHistogram::OnPaint( const FPaintArgs& Args, const FGeometry& AllottedGeom
 
 		// Bottom - X-Axes numbers, starting from MinValue
 		const FString XLabel = FString::Printf(TEXT("%.0f"), Description.MinValue + Index*Description.Interval);
-		float FontCharWidth = FontMeasureService->Measure(XLabel, SummaryFont).X;
+		float FontCharWidth = static_cast<float>(FontMeasureService->Measure(XLabel, SummaryFont).X);
 		FSlateDrawElement::MakeText(
 			OutDrawElements, 
 			LayerId, 
@@ -118,7 +118,7 @@ int32 SHistogram::OnPaint( const FPaintArgs& Args, const FGeometry& AllottedGeom
 	LayerId++;
 
 	// draw the horizontal lines
-	float CountY = (AllottedGeometry.GetLocalSize().Y-LabelBuffer*2.0f) / 4;
+	float CountY = (static_cast<float>(AllottedGeometry.GetLocalSize().Y)-LabelBuffer*2.0f) / 4.0f;
 	float StartY = LabelBuffer;
 	for (int32 Index = 0; Index < 5; ++Index)
 	{
@@ -137,7 +137,7 @@ int32 SHistogram::OnPaint( const FPaintArgs& Args, const FGeometry& AllottedGeom
 
 		// Bottom - Y-Axes numbers, starting from 0
 		const FString YLabel = FString::Printf(TEXT("%i"), Description.Normalize ? 25 * (4-Index) : Description.GetTotalCount() / 4 * Index);
-		float FontCharWidth = FontMeasureService->Measure(YLabel, SummaryFont).X;
+		float FontCharWidth = static_cast<float>(FontMeasureService->Measure(YLabel, SummaryFont).X);
 		FSlateDrawElement::MakeText
 			(
 			OutDrawElements, 
@@ -155,7 +155,7 @@ int32 SHistogram::OnPaint( const FPaintArgs& Args, const FGeometry& AllottedGeom
 	for (int32 Index = 0; Index < Description.GetBinCount(); ++Index)
 	{
 		float MarkerPosX = StartX + Index * CountX;
-		float SizeY = (float)Description.GetCount(Index) / (float)Description.GetTotalCount() * (AllottedGeometry.GetLocalSize().Y - LabelBuffer*2.0f);
+		float SizeY = (float)Description.GetCount(Index) / (float)Description.GetTotalCount() * ((float)AllottedGeometry.GetLocalSize().Y - LabelBuffer*2.0f);
 		FSlateDrawElement::MakeBox(
 			OutDrawElements,
 			LayerId,
diff --git a/Engine/Source/Developer/Profiler/Private/Widgets/SMultiDumpBrowser.cpp b/Engine/Source/Developer/Profiler/Private/Widgets/SMultiDumpBrowser.cpp
--- a/Engine/Source/Developer/Profiler/Private/Widgets/SMultiDumpBrowser.cpp
+++ b/Engine/Source/Developer/Profiler/Private/Widgets/SMultiDumpBrowser.cpp
@@ -125,7 +125,7 @@ void SMultiDumpBrowser::FindTotalsForPrefilterRecursive(TSharedPtr<FEventGraphSa
 	FString EventName = Event->_StatName.ToString().ToUpper();
 	if (EventName.Contains(TotalsFilteringText.ToUpper()))
 	{
-		OutTotalTime += Event->_InclusiveTimeMS;
+		OutTotalTime += static_cast<float>(Event->_InclusiveTimeMS);
 	}
 	else
 	{
diff --git a/Engine/Source/Developer/Profiler/Private/Widgets/SProfilerMiniView.cpp b/Engine/Source/Developer/Profiler/Private/Widgets/SProfilerMiniView.cpp
--- a/Engine/Source/Developer/Profiler/Private/Widgets/SProfilerMiniView.cpp
+++ b/Engine/Source/Developer/Profiler/Private/Widgets/SProfilerMiniView.cpp
@@ -107,12 +107,12 @@ int32 SProfilerMiniView::OnPaint( const FPaintArgs& Args, const FGeometry& Allot
 
 	PaintState = new((void*)PaintStateMemory) FSlateOnPaintState( AllottedGeometry, MyCullingRect, OutDrawElements, LayerId, InWidgetStyle, DrawEffects );
 
-	const float MiniViewSizeX = AllottedGeometry.Size.X;
-	const float MiniViewSizeY = AllottedGeometry.Size.Y;
+	const float MiniViewSizeX = static_cast<float>(AllottedGeometry.Size.X);
+	const float MiniViewSizeY = static_cast<float>(AllottedGeometry.Size.Y);
 
 	const TSharedRef< FSlateFontMeasure > FontMeasureService = FSlateApplication::Get().GetRenderer()->GetFontMeasureService();
 	FSlateFontInfo SummaryFont = FCoreStyle::GetDefaultFontStyle("Regular", 8);
-	const float MaxFontCharHeight = FontMeasureService->Measure( TEXT( "!" ), SummaryFont ).Y;
+	const float MaxFontCharHeight = static_cast<float>(FontMeasureService->Measure( TEXT( "!" ), SummaryFont ).Y);
 
 	// Draw background.
 	FSlateDrawElement::MakeBox
@@ -195,8 +195,8 @@ int32 SProfilerMiniView::OnPaint( const FPaintArgs& Args, const FGeometry& Allot
 		
 		const int32 MaxFrameIndex = AllFrames.Num() - 1;
 		
-		const int32 SelectionBoxX0 = FMath::TruncToInt( FrameIndexToPosition( SelectionBoxFrameStart ) );
-		const int32 SelectionBoxX1 = FMath::TruncToInt(FrameIndexToPosition(SelectionBoxFrameEnd + 1));
+		const double SelectionBoxX0 = FMath::TruncToInt( FrameIndexToPosition( SelectionBoxFrameStart ) );
+		const double SelectionBoxX1 = FMath::TruncToInt( FrameIndexToPosition( SelectionBoxFrameEnd + 1 ));
 
 		if( SelectionBoxFrameStart > 0 )
 		{
@@ -204,7 +204,7 @@ int32 SProfilerMiniView::OnPaint( const FPaintArgs& Args, const FGeometry& Allot
 			(
 				OutDrawElements,
 				LayerId,
-				AllottedGeometry.ToPaintGeometry( FVector2D( 0.0f, 0.0f ), FVector2D( (float)SelectionBoxX0, AllottedGeometry.Size.Y ) ),
+				AllottedGeometry.ToPaintGeometry( FVector2D( 0.0, 0.0 ), FVector2D( SelectionBoxX0, AllottedGeometry.Size.Y ) ),
 				&SolidWhiteBrush,
 				DrawEffects,
 				FColorList::Grey.WithAlpha( 192 )
@@ -217,7 +217,7 @@ int32 SProfilerMiniView::OnPaint( const FPaintArgs& Args, const FGeometry& Allot
 			(
 				OutDrawElements,
 				LayerId,
-				AllottedGeometry.ToPaintGeometry( FVector2D( SelectionBoxX1, 0.0f ), FVector2D( MiniViewSizeX - SelectionBoxX1, AllottedGeometry.Size.Y ) ),
+				AllottedGeometry.ToPaintGeometry( FVector2D( SelectionBoxX1, 0.0 ), FVector2D( MiniViewSizeX - SelectionBoxX1, AllottedGeometry.Size.Y ) ),
 				&SolidWhiteBrush,
 				DrawEffects,
 				FColorList::Grey.WithAlpha( 192 )
@@ -225,15 +225,15 @@ int32 SProfilerMiniView::OnPaint( const FPaintArgs& Args, const FGeometry& Allot
 		}
 
 		// Draw the filler, to hide the difference between window's width and mini-view samples' width.
-		const float FillerSizeX = MiniViewSizeX - CurrentSamplePosX;
-		if( FillerSizeX > 0.0f )
+		const double FillerSizeX = MiniViewSizeX - CurrentSamplePosX;
+		if( FillerSizeX > 0.0 )
 		{
-			const float FillerPosX0 = MiniViewSizeX - FillerSizeX;
+			const double FillerPosX0 = MiniViewSizeX - FillerSizeX;
 			FSlateDrawElement::MakeBox
 			(
 				OutDrawElements,
 				LayerId,
-				AllottedGeometry.ToPaintGeometry( FVector2D( FillerPosX0, 0.0f ), FVector2D( FillerSizeX + 1.0f, AllottedGeometry.Size.Y ) ),
+				AllottedGeometry.ToPaintGeometry( FVector2D( FillerPosX0, 0.0 ), FVector2D( FillerSizeX + 1.0, AllottedGeometry.Size.Y ) ),
 				&SolidWhiteBrush,
 				DrawEffects,
 				// #Profiler: 2014-04-09 How to get this color from Slate?
@@ -247,7 +247,7 @@ int32 SProfilerMiniView::OnPaint( const FPaintArgs& Args, const FGeometry& Allot
 		(
 			OutDrawElements,
 			LayerId,
-			AllottedGeometry.ToPaintGeometry( FVector2D( SelectionBoxX0, 0.0f ), FVector2D( SelectionBoxX1 - SelectionBoxX0, AllottedGeometry.Size.Y ) ),
+			AllottedGeometry.ToPaintGeometry( FVector2D( SelectionBoxX0, 0.0 ), FVector2D( SelectionBoxX1 - SelectionBoxX0, AllottedGeometry.Size.Y ) ),
 			FProfilerStyle::Get().GetBrush( "PlainBorder" ),
 			DrawEffects,
 			FColorList::Green
@@ -261,19 +261,19 @@ int32 SProfilerMiniView::OnPaint( const FPaintArgs& Args, const FGeometry& Allot
 		const float MarkerPosY = 4.0f;
 
 		// Upper left
-		DrawText( TEXT( "Rendering thread" ), SummaryFont, FVector2D( MarkerPosX, MarkerPosY ), RenderThreadColor, FColor::Black, FVector2D( 1.0f, 1.0f ) );
+		DrawText( TEXT( "Rendering thread" ), SummaryFont, FVector2D( MarkerPosX, MarkerPosY ), RenderThreadColor, FColor::Black, FVector2D( 1.0, 1.0 ) );
 		// Down left
-		DrawText( TEXT( "Game thread" ), SummaryFont, FVector2D( MarkerPosX, MiniViewSizeY - MarkerPosY - MaxFontCharHeight ), GameThreadColor, FColor::Black, FVector2D( 1.0f, 1.0f ) );
+		DrawText( TEXT( "Game thread" ), SummaryFont, FVector2D( MarkerPosX, MiniViewSizeY - MarkerPosY - MaxFontCharHeight ), GameThreadColor, FColor::Black, FVector2D( 1.0, 1.0 ) );
 
 		// Upper right
 		const FString ThreadTimeMax = FString::Printf( TEXT( "%5.2f MS" ), MaxFrameTime );
-		const float ThreadTimeMaxSizeX = FontMeasureService->Measure( ThreadTimeMax, SummaryFont ).X;
-		DrawText( ThreadTimeMax, SummaryFont, FVector2D( MiniViewSizeX - ThreadTimeMaxSizeX - MarkerPosX, MarkerPosY ), FColor::White, FColor::Black, FVector2D( 1.0f, 1.0f ) );
+		const double ThreadTimeMaxSizeX = FontMeasureService->Measure( ThreadTimeMax, SummaryFont ).X;
+		DrawText( ThreadTimeMax, SummaryFont, FVector2D( MiniViewSizeX - ThreadTimeMaxSizeX - MarkerPosX, MarkerPosY ), FColor::White, FColor::Black, FVector2D( 1.0, 1.0 ) );
 
 		// Down right
 		const FString ThreadTimeMin = TEXT( "0.0 MS" );
-		const float ThreadTimeMinSizeX = FontMeasureService->Measure( ThreadTimeMin, SummaryFont ).X;
-		DrawText( ThreadTimeMin, SummaryFont, FVector2D( MiniViewSizeX - ThreadTimeMinSizeX - MarkerPosX, MiniViewSizeY - MarkerPosY - MaxFontCharHeight ), FColor::White, FColor::Black, FVector2D( 1.0f, 1.0f ) );
+		const double ThreadTimeMinSizeX = FontMeasureService->Measure( ThreadTimeMin, SummaryFont ).X;
+		DrawText( ThreadTimeMin, SummaryFont, FVector2D( MiniViewSizeX - ThreadTimeMinSizeX - MarkerPosX, MiniViewSizeY - MarkerPosY - MaxFontCharHeight ), FColor::White, FColor::Black, FVector2D( 1.0, 1.0 ) );
 	
 	}
 
@@ -381,7 +381,7 @@ FReply SProfilerMiniView::OnMouseButtonDown( const FGeometry& MyGeometry, const
 			else
 			{
 				// Clicked outside the selection box, so move the selection box to that position.
-				DistanceDragged = MousePositionOnButtonDown.X;
+				DistanceDragged = static_cast<float>(MousePositionOnButtonDown.X);
 			}
 
 			if( bCanBeStartDragged || bCanBeEndDragged )
@@ -451,9 +451,9 @@ FReply SProfilerMiniView::OnMouseMove( const FGeometry& MyGeometry, const FPoint
 	if( IsReady() )
 	{
 		const FVector2D MousePosition = MyGeometry.AbsoluteToLocal( MouseEvent.GetScreenSpacePosition() );
-		HoveredFrameIndex = PositionToFrameIndex( MousePosition.X );
+		HoveredFrameIndex = PositionToFrameIndex( static_cast<float>(MousePosition.X) );
 
-		const float CursorPosXDelta = MouseEvent.GetCursorDelta().X;
+		const float CursorPosXDelta = static_cast<float>(MouseEvent.GetCursorDelta().X);
 
 		if( MouseEvent.IsMouseButtonDown( EKeys::LeftMouseButton ) )
 		{
@@ -492,8 +492,8 @@ FReply SProfilerMiniView::OnMouseMove( const FGeometry& MyGeometry, const FPoint
 		}
 		else
 		{
-			const float StartEdgeDistance = FrameIndexToPosition(SelectionBoxFrameStart) - MousePosition.X;
-			const float EndEdgeDistance = MousePosition.X - FrameIndexToPosition(SelectionBoxFrameEnd);
+			const float StartEdgeDistance = FrameIndexToPosition(SelectionBoxFrameStart) - static_cast<float>(MousePosition.X);
+			const float EndEdgeDistance = static_cast<float>(MousePosition.X) - FrameIndexToPosition(SelectionBoxFrameEnd);
 
 			bCanBeStartDragged = StartEdgeDistance < (float)MOUSE_SNAP_DISTANCE && StartEdgeDistance > 0.0f;
 			bCanBeEndDragged = EndEdgeDistance < (float)MOUSE_SNAP_DISTANCE && EndEdgeDistance > 0.0f;
@@ -657,7 +657,7 @@ void SProfilerMiniView::ProcessData()
 
 	UpdateNumPixelsPerSample();
 
-	const int32 NumMiniViewSamples = FMath::TruncToInt(ThisGeometry.Size.X / GetNumPixelsPerSample());
+	const int32 NumMiniViewSamples = FMath::TruncToInt(static_cast<float>(ThisGeometry.Size.X) / GetNumPixelsPerSample());
 	MiniViewSamples.Reset( NumMiniViewSamples );
 
 	MiniViewSamples.AddZeroed( NumMiniViewSamples );
diff --git a/Engine/Source/Developer/Profiler/Private/Widgets/SProfilerMiniView.h b/Engine/Source/Developer/Profiler/Private/Widgets/SProfilerMiniView.h
--- a/Engine/Source/Developer/Profiler/Private/Widgets/SProfilerMiniView.h
+++ b/Engine/Source/Developer/Profiler/Private/Widgets/SProfilerMiniView.h
@@ -301,7 +301,7 @@ protected:
 
 	void UpdateNumPixelsPerSample()
 	{
-		NumPixelsPerSample = FMath::Max( (float)MIN_NUM_PIXELS_PER_SAMPLE, ThisGeometry.GetLocalSize().X / AllFrames.Num() );
+		NumPixelsPerSample = FMath::Max( (float)MIN_NUM_PIXELS_PER_SAMPLE, (float)ThisGeometry.GetLocalSize().X / AllFrames.Num() );
 	}
 
 	const float GetNumPixelsPerSample() const
diff --git a/Engine/Source/Developer/Profiler/Private/Widgets/SProfilerThreadView.cpp b/Engine/Source/Developer/Profiler/Private/Widgets/SProfilerThreadView.cpp
--- a/Engine/Source/Developer/Profiler/Private/Widgets/SProfilerThreadView.cpp
+++ b/Engine/Source/Developer/Profiler/Private/Widgets/SProfilerThreadView.cpp
@@ -197,7 +197,7 @@ void SProfilerThreadView::DrawFramesBackgroundAndTimelines() const
 			);
 
 			// Check if this frame time marker is inside the visible area.
-			const float LocalPositionXPx = PositionPx.X + SizePx.X;
+			const float LocalPositionXPx = static_cast<float>(PositionPx.X + SizePx.X);
 			if( LocalPositionXPx < 0.0f || LocalPositionXPx > PaintState->Size2D().X )
 			{
 				continue;
@@ -325,8 +325,8 @@ void SProfilerThreadView::DrawUIStackNodes() const
 			}
 
 			// Update position of the text to be always visible and try to center it.
-			const float StatNameWidthPx = PaintState->FontMeasureService->Measure( StringStatName, PaintState->SummaryFont8 ).X;
-			const float StatNameWithTimeWidthPx = PaintState->FontMeasureService->Measure( StringStatNameWithTime, PaintState->SummaryFont8 ).X;
+			const float StatNameWidthPx = static_cast<float>(PaintState->FontMeasureService->Measure( StringStatName, PaintState->SummaryFont8 ).X);
+			const float StatNameWithTimeWidthPx = static_cast<float>(PaintState->FontMeasureService->Measure( StringStatNameWithTime, PaintState->SummaryFont8 ).X);
 			const float TextAreaWidthPx = ClippedNodeRect.GetSize().X;
 
 			bool bUseShortVersion = true;
@@ -367,10 +367,10 @@ void SProfilerThreadView::DrawFrameMarkers() const
 	{
 		if( ThreadNode.StatName == NAME_GameThread )
 		{
-			const float MarkerPosXPx = ThreadNode.GetLocalPosition( ThreadViewOffsetPx, 0.0f ).X + ThreadNode.WidthPx;
+			const double MarkerPosXPx = ThreadNode.GetLocalPosition( ThreadViewOffsetPx, 0.0f ).X + ThreadNode.WidthPx;
 
 			// Check if this frame time marker is inside the visible area.
-			if( MarkerPosXPx < 0.0f || MarkerPosXPx > PaintState->Size2D().X )
+			if( MarkerPosXPx < 0.0 || MarkerPosXPx > PaintState->Size2D().X )
 			{
 				continue;
 			}
@@ -379,7 +379,7 @@ void SProfilerThreadView::DrawFrameMarkers() const
 			const FString FrameIndexStr = FString::Printf( TEXT( "%i" ), ThreadNode.FrameIndex );
 			const FString FrameTimesStr = FString::Printf( TEXT( "%.4f [%.4f] MS" ), ThreadNode.CycleCountersEndTimeMS, ThreadNode.GetDurationMS() );
 
-			float MarkerPosYPx = PaintState->Size2D().Y - 2 * PaintState->SummaryFont8Height;
+			double MarkerPosYPx = PaintState->Size2D().Y - 2 * PaintState->SummaryFont8Height;
 			DrawText( FrameIndexStr, PaintState->SummaryFont8, FVector2D( MarkerPosXPx, MarkerPosYPx ), FColorList::SkyBlue, FColorList::Black, FVector2D( 1.0f, 1.0f ) );
 
 			MarkerPosYPx += PaintState->SummaryFont8Height;
@@ -397,7 +397,7 @@ void SProfilerThreadView::DrawFrameMarkers() const
 		const FString TimelineStr = FString::Printf( TEXT( "%.4f MS" ), TimelinePosXPx / NumPixelsPerMillisecond );
 
 		// Draw time line text.
-		float MarkerPosYPx = PaintState->Size2D().Y - 3 * PaintState->SummaryFont8Height;
+		double MarkerPosYPx = PaintState->Size2D().Y - 3 * PaintState->SummaryFont8Height;
 		DrawText( TimelineStr, PaintState->SummaryFont8, FVector2D( TimelinePosXPx - ThreadViewOffsetPx, MarkerPosYPx ), FColorList::LimeGreen, FColorList::Black, FVector2D( 1.0f, 1.0f ) );
 	}
 
@@ -560,8 +560,8 @@ FReply SProfilerThreadView::OnMouseMove( const FGeometry& MyGeometry, const FPoi
 		HoveredPositionX = 0.0;//PositionToFrameIndex( LocalMousePosition.X );
 		HoveredPositionY = 0.0;
 
-		const float CursorPosXDelta = -MouseEvent.GetCursorDelta().X;
-		const float ScrollSpeed = 1.0f / ZoomFactorX;
+		const double CursorPosXDelta = -MouseEvent.GetCursorDelta().X;
+		const double ScrollSpeed = 1.0 / ZoomFactorX;
 
 		if( MouseEvent.IsMouseButtonDown( EKeys::LeftMouseButton ) )
 		{
diff --git a/Engine/Source/Developer/Profiler/Private/Widgets/SProfilerThreadView.h b/Engine/Source/Developer/Profiler/Private/Widgets/SProfilerThreadView.h
--- a/Engine/Source/Developer/Profiler/Private/Widgets/SProfilerThreadView.h
+++ b/Engine/Source/Developer/Profiler/Private/Widgets/SProfilerThreadView.h
@@ -95,7 +95,7 @@ class SProfilerThreadView
 			, DrawEffects( InDrawEffects )
 			, FontMeasureService( FSlateApplication::Get().GetRenderer()->GetFontMeasureService() )
 			, SummaryFont8( FCoreStyle::GetDefaultFontStyle("Regular", 8) )
-			, SummaryFont8Height( FontMeasureService->Measure( TEXT( "!" ), SummaryFont8 ).Y )
+			, SummaryFont8Height( static_cast<float>(FontMeasureService->Measure( TEXT( "!" ), SummaryFont8 ).Y) )
 		{}
 
 		const FVector2D& Size2D() const
diff --git a/Engine/Source/Developer/ProfilerClient/Private/ProfilerClientManager.cpp b/Engine/Source/Developer/ProfilerClient/Private/ProfilerClientManager.cpp
--- a/Engine/Source/Developer/ProfilerClient/Private/ProfilerClientManager.cpp
+++ b/Engine/Source/Developer/ProfilerClient/Private/ProfilerClientManager.cpp
@@ -239,7 +239,7 @@ protected:
 
 		FProfilerDataFrame& DataFrame = LoadConnection->CurrentData;
 
-		DataFrame.Frame = Frame;
+		DataFrame.Frame = static_cast<uint32>(Frame);
 		DataFrame.FrameStart = 0.0;
 		DataFrame.CountAccumulators.Reset();
 		DataFrame.CycleGraphs.Reset();
@@ -618,7 +618,7 @@ void FProfilerClientManager::HandleServicePingMessage(const FProfilerServicePing
 bool FProfilerClientManager::HandleTicker(float DeltaTime)
 {
 #if STATS
-	if (PendingInstances.Num() > 0 && FDateTime::Now() > LastPingTime + FTimespan(DeltaTime))
+	if (PendingInstances.Num() > 0 && FDateTime::Now() > LastPingTime + FTimespan::FromSeconds(DeltaTime))
 	{
 		TArray<FGuid> Instances;
 		Instances.Append(PendingInstances);
@@ -1022,7 +1022,7 @@ void FServiceConnection::GenerateAccumulators(TArray<FStatMessage>& Stats, TArra
 				// add a count accumulator
 				FProfilerCountAccumulator Data;
 				Data.StatId = StatId;
-				Data.Value = StatMessage.GetValue_int64();
+				Data.Value = static_cast<uint32>(StatMessage.GetValue_int64());
 				CountAccumulators.Add(Data);
 			}
 			else if (StatMessage.NameAndInfo.GetField<EStatDataType>() == EStatDataType::ST_double)
@@ -1030,7 +1030,7 @@ void FServiceConnection::GenerateAccumulators(TArray<FStatMessage>& Stats, TArra
 				// add a float accumulator
 				FProfilerFloatAccumulator Data;
 				Data.StatId = StatId;
-				Data.Value = StatMessage.GetValue_double();
+				Data.Value = static_cast<float>(StatMessage.GetValue_double());
 				FloatAccumulators.Add(Data);
 
 				const FName StatName = StatMessage.NameAndInfo.GetRawName();
@@ -1060,7 +1060,7 @@ void FServiceConnection::CreateGraphRecursively(const FRawStatStackNode* Root, F
 		else
 		{
 			Graph.CallsPerFrame = 1;
-			Graph.Value = Root->Meta.GetValue_int64();
+			Graph.Value = static_cast<uint32>(Root->Meta.GetValue_int64());
 		}
 	}
 
@@ -1117,7 +1117,7 @@ void FServiceConnection::GenerateProfilerDataFrame()
 {
 	SCOPE_CYCLE_COUNTER(STAT_PC_GenerateDataFrame);
 	FProfilerDataFrame& DataFrame = CurrentData;
-	DataFrame.Frame = CurrentThreadState.CurrentGameFrame;
+	DataFrame.Frame = static_cast<uint32>(CurrentThreadState.CurrentGameFrame);
 	DataFrame.FrameStart = 0.0;
 	DataFrame.CountAccumulators.Reset();
 	DataFrame.CycleGraphs.Reset();
