Fixing compile errors ahead of enabling TObjectPtr GC barrier

diff --git a/Engine/Source/Developer/Profiler/Profiler.Build.cs b/Engine/Source/Developer/Profiler/Profiler.Build.cs
--- a/Engine/Source/Developer/Profiler/Profiler.Build.cs
+++ b/Engine/Source/Developer/Profiler/Profiler.Build.cs
@@ -10,6 +10,7 @@ public class Profiler : ModuleRules
 		(
 			new string[] {
 				"Core",
+				"CoreUObject",
 				"ApplicationCore",
                 "InputCore",
 				"RHI",
