Enabled UnsafeTypeCast warnings as errors in Profiler, ProfilerClient, PerfCounters modules.

diff --git a/Engine/Source/Developer/Profiler/Profiler.Build.cs b/Engine/Source/Developer/Profiler/Profiler.Build.cs
--- a/Engine/Source/Developer/Profiler/Profiler.Build.cs
+++ b/Engine/Source/Developer/Profiler/Profiler.Build.cs
@@ -6,13 +6,14 @@ public class Profiler : ModuleRules
 {
 	public Profiler( ReadOnlyTargetRules Target ) : base(Target)
 	{
-		PublicDependencyModuleNames.AddRange
-		(
+		UnsafeTypeCastWarningLevel = WarningLevel.Error;
+
+		PublicDependencyModuleNames.AddRange(
 			new string[] {
 				"Core",
 				"CoreUObject",
 				"ApplicationCore",
-                "InputCore",
+				"InputCore",
 				"RHI",
 				"RenderCore",
 				"Slate",
@@ -21,14 +22,14 @@ public class Profiler : ModuleRules
 			}
 		);
 
-        if (Target.bBuildEditor)
-        {
-            PrivateDependencyModuleNames.AddRange(
-                new string[] {
+		if (Target.bBuildEditor)
+		{
+			PrivateDependencyModuleNames.AddRange(
+				new string[] {
 					"Engine",
 				}
-            );
-        }
+			);
+		}
 
 		PrivateDependencyModuleNames.AddRange(
 			new string[] {
diff --git a/Engine/Source/Developer/ProfilerClient/Private/ProfilerClientManager.cpp b/Engine/Source/Developer/ProfilerClient/Private/ProfilerClientManager.cpp
--- a/Engine/Source/Developer/ProfilerClient/Private/ProfilerClientManager.cpp
+++ b/Engine/Source/Developer/ProfilerClient/Private/ProfilerClientManager.cpp
@@ -721,7 +721,7 @@ bool FProfilerClientManager::HandleMessagesTicker(float DeltaTime)
 				Connection.ReceivedData.Remove(FrameNum);
 
 				// see if we need to yield
-				double DurationSeconds = (FPlatformTime::Cycles64() - StartTimeCycles) * FPlatformTime::GetSecondsPerCycle64();
+				double DurationSeconds = static_cast<double>(FPlatformTime::Cycles64() - StartTimeCycles) * FPlatformTime::GetSecondsPerCycle64();
 				if (DurationSeconds > MaxDurationSeconds)
 				{
 					UE_CLOG(Index < Frames.Num()-1, LogProfilerClient, Verbose, TEXT("Over time - %d/%d frames processed for connection %s"), Index, Frames.Num() - Index, *Connection.InstanceId.ToString() );
diff --git a/Engine/Source/Developer/ProfilerClient/ProfilerClient.Build.cs b/Engine/Source/Developer/ProfilerClient/ProfilerClient.Build.cs
--- a/Engine/Source/Developer/ProfilerClient/ProfilerClient.Build.cs
+++ b/Engine/Source/Developer/ProfilerClient/ProfilerClient.Build.cs
@@ -6,18 +6,18 @@ namespace UnrealBuildTool.Rules
 	{
 		public ProfilerClient(ReadOnlyTargetRules Target) : base(Target)
 		{
+			UnsafeTypeCastWarningLevel = WarningLevel.Error;
+
 			PublicIncludePathModuleNames.Add("ProfilerService");
 
 			PublicDependencyModuleNames.AddRange(
-				new string[]
-				{
+				new string[] {
 					"Core",
 				}
 			);
 
 			PrivateDependencyModuleNames.AddRange(
-				new string[]
-				{
+				new string[] {
 					"CoreUObject",
 					"ProfilerMessages",
 				}
@@ -26,7 +26,7 @@ namespace UnrealBuildTool.Rules
 			PrivateIncludePathModuleNames.AddRange(
 				new string[] {
 					"Messaging",
-                    "MessagingCommon",
+					"MessagingCommon",
 				}
 			);
 		}
diff --git a/Engine/Source/Runtime/PerfCounters/PerfCounters.Build.cs b/Engine/Source/Runtime/PerfCounters/PerfCounters.Build.cs
--- a/Engine/Source/Runtime/PerfCounters/PerfCounters.Build.cs
+++ b/Engine/Source/Runtime/PerfCounters/PerfCounters.Build.cs
@@ -4,14 +4,16 @@ using UnrealBuildTool;
 
 public class PerfCounters : ModuleRules
 {
-    public PerfCounters(ReadOnlyTargetRules Target) : base(Target)
-    {
-        PrivateDependencyModuleNames.AddRange(
-            new string[] {
+	public PerfCounters(ReadOnlyTargetRules Target) : base(Target)
+	{
+		UnsafeTypeCastWarningLevel = WarningLevel.Error;
+
+		PrivateDependencyModuleNames.AddRange(
+			new string[] {
 				"Core",
 				"Json",
-                "HTTPServer"
-            }
-        );
+				"HTTPServer"
+			}
+		);
 	}
 }
