Fixed "implicit conversion" errors in Profiler module.

#rnx

diff --git a/Engine/Source/Developer/Profiler/Private/ProfilerFPSAnalyzer.cpp b/Engine/Source/Developer/Profiler/Private/ProfilerFPSAnalyzer.cpp
--- a/Engine/Source/Developer/Profiler/Private/ProfilerFPSAnalyzer.cpp
+++ b/Engine/Source/Developer/Profiler/Private/ProfilerFPSAnalyzer.cpp
@@ -26,7 +26,7 @@ void FFPSAnalyzer::AddSample(float FPSSample)
 	Samples.Add(FPSSample);
 
 	float DeltaSeconds = 1.0f / FPSSample;
-	int32 Index = FPlatformMath::FloorToInt(FPSSample / Interval);
+	int32 Index = FPlatformMath::FloorToInt(FPSSample / (float)Interval);
 	Index = Index >= Histogram.Num() ? Histogram.Num()-1 : Index;
 	
 	Histogram[Index].Count++;
@@ -86,7 +86,7 @@ int32 FFPSAnalyzer::GetTotalCount()
 
 int32 FFPSAnalyzer::GetCount(float InMinVal, float InMaxVal)
 {
-	int32 Index = FPlatformMath::FloorToInt(InMinVal / Interval);
+	int32 Index = FPlatformMath::FloorToInt(InMinVal / (float)Interval);
 	Index = Index >= Histogram.Num() ? Histogram.Num()-1 : Index;
 
 	return Histogram[Index].Count;
diff --git a/Engine/Source/Developer/Profiler/Private/ProfilerManager.cpp b/Engine/Source/Developer/Profiler/Private/ProfilerManager.cpp
--- a/Engine/Source/Developer/Profiler/Private/ProfilerManager.cpp
+++ b/Engine/Source/Developer/Profiler/Private/ProfilerManager.cpp
@@ -302,8 +302,8 @@ bool FProfilerManager::Tick( float DeltaTime )
 		const double MBInv = 1.0 / 1024.0 / 1024.0;
 		const SIZE_T DiffPhys = FPlatformMemory::GetStats().UsedPhysical - StartUsedPhysical;
 
-		const double SessionMemory = MBInv*ProfilerSession->GetMemoryUsage();
-		const double PhysMemory = MBInv*DiffPhys;
+		const double SessionMemory = MBInv * (double)ProfilerSession->GetMemoryUsage();
+		const double PhysMemory = MBInv * (double)DiffPhys;
 
 		UE_LOG( LogStats, VeryVerbose, TEXT( "ProfilerSession: %6.2f MB (%6.2f MB) # (%6.2f MB) / %7u -> %4u" ), 
 			SessionMemory,
diff --git a/Engine/Source/Developer/Profiler/Private/ProfilerRawStatsForThreadView.cpp b/Engine/Source/Developer/Profiler/Private/ProfilerRawStatsForThreadView.cpp
--- a/Engine/Source/Developer/Profiler/Private/ProfilerRawStatsForThreadView.cpp
+++ b/Engine/Source/Developer/Profiler/Private/ProfilerRawStatsForThreadView.cpp
@@ -235,14 +235,14 @@ void FRawProfilerSession::PrepareLoading()
 				}
 
 				const int64 CurrentPos = FileReader->Tell();
-				const int32 PctPos = int32(100.0f*CurrentPos/FileSize);
+				const int32 PctPos = int32(100.0f * (float)CurrentPos / (float)FileSize);
 
 				UE_LOG( LogStats, Log, TEXT( "%3i Processing FStatPacket: Frame %5i for thread %5i with %6i messages (%.1f MB)" ), 
 					PctPos, 
 					StatPacket->Frame, 
 					StatPacket->ThreadId, 
 					StatPacket->StatMessages.Num(), 
-					StatPacket->StatMessages.GetAllocatedSize()/1024.0f/1024.0f );
+					(double)StatPacket->StatMessages.GetAllocatedSize()/1024.0/1024.0 );
 
 				const int64 PacketSize = StatPacket->StatMessages.GetAllocatedSize();
 				TotalPacketSize += PacketSize;
@@ -251,8 +251,8 @@ void FRawProfilerSession::PrepareLoading()
 		}
 
 		UE_LOG( LogStats, Log, TEXT( "TotalPacketSize: %.1f MB, Max: %1f MB" ), 
-			TotalPacketSize/1024.0f/1024.0f, 
-			MaximumPacketSize/1024.0f/1024.0f );
+			(double)TotalPacketSize/1024.0/1024.0, 
+			(double)MaximumPacketSize/1024.0/1024.0 );
 
 		TArray<int64> Frames;
 		CombinedHistory.GenerateKeyArray(Frames);
diff --git a/Engine/Source/Developer/Profiler/Private/ProfilerSession.cpp b/Engine/Source/Developer/Profiler/Private/ProfilerSession.cpp
--- a/Engine/Source/Developer/Profiler/Private/ProfilerSession.cpp
+++ b/Engine/Source/Developer/Profiler/Private/ProfilerSession.cpp
@@ -592,9 +592,10 @@ void FProfilerSession::PopulateHierarchy_Recurrent
 		ChildrenDurationCycles += ChildDurationCycles;
 	}
 
-	const uint32 SelfTimeCycles = ParentDurationCycles - ChildrenDurationCycles;
-	if( SelfTimeCycles > 0.0f && ParentGraph.Children.Num() > 0 )
+	if (ParentDurationCycles > ChildrenDurationCycles && ParentGraph.Children.Num() > 0)
 	{
+		const uint32 SelfTimeCycles = ParentDurationCycles - ChildrenDurationCycles;
+
 		// Create a fake stat that represents this profiler sample's exclusive time.
 		// This is required if we want to create correct combined event graphs later.
 		DataProvider->AddHierarchicalSample
@@ -638,7 +639,7 @@ float FProfilerSession::GetProgress() const
 {
 	if (NumFrames > 0)
 	{
-		return (float)NumFramesProcessed / NumFrames;
+		return (float)NumFramesProcessed / (float)NumFrames;
 	}
 	return 0.0f;
 }
diff --git a/Engine/Source/Developer/Profiler/Private/Widgets/SDataGraph.cpp b/Engine/Source/Developer/Profiler/Private/Widgets/SDataGraph.cpp
--- a/Engine/Source/Developer/Profiler/Private/Widgets/SDataGraph.cpp
+++ b/Engine/Source/Developer/Profiler/Private/Widgets/SDataGraph.cpp
@@ -275,8 +275,8 @@ void SDataGraph::UpdateState()
 		GraphOffset = FMath::Clamp( GraphOffset, 0, FMath::Max(NumDataPoints-NumVisiblePoints,0) );
 		
 		DataTotalTimeMS = FirstTrackedStat->GraphDataSource->GetTotalTimeMS();
-		VisibleTimeMS = NumVisiblePoints * FTimeAccuracy::AsFrameTime( TimeBasedAccuracy );
-		GraphOffsetMS = GraphOffset * FTimeAccuracy::AsFrameTime( TimeBasedAccuracy );
+		VisibleTimeMS = (float)NumVisiblePoints * FTimeAccuracy::AsFrameTime( TimeBasedAccuracy );
+		GraphOffsetMS = (float)GraphOffset * FTimeAccuracy::AsFrameTime( TimeBasedAccuracy );
 	}
 	else
 	{
@@ -457,7 +457,7 @@ int32 SDataGraph::OnPaint( const FPaintArgs& Args, const FGeometry& AllottedGeom
 				for( uint32 GraphStartIndex = (uint32)GraphOffset; GraphStartIndex < (uint32)GraphRangeEndIndex; GraphStartIndex++ )
 				{
 					const float Value = GraphDataSource->GetValueFromIndex( GraphStartIndex );
-					const float XPos = DistanceBetweenPoints*(float)GraphPoints.Num();
+					const float XPos = (float)DistanceBetweenPoints * (float)GraphPoints.Num();
 					const float YPos = FMath::Clamp( static_cast<float>(AllottedGeometry.Size.Y) - GraphYScale*Value*UnitTypeScale, 0.0f, static_cast<float>(AllottedGeometry.Size.Y) );
 					GraphPoints.Add( FVector2D(XPos,YPos) );
 				}
diff --git a/Engine/Source/Developer/Profiler/Private/Widgets/SProfilerMiniView.h b/Engine/Source/Developer/Profiler/Private/Widgets/SProfilerMiniView.h
--- a/Engine/Source/Developer/Profiler/Private/Widgets/SProfilerMiniView.h
+++ b/Engine/Source/Developer/Profiler/Private/Widgets/SProfilerMiniView.h
@@ -293,7 +293,7 @@ protected:
 	*/
 	const float FrameIndexToPosition( const int32 FrameIndex ) const
 	{
-		return FrameIndex*NumPixelsPerFrame;
+		return (float)FrameIndex * NumPixelsPerFrame;
 	}
 
 protected:
@@ -301,7 +301,7 @@ protected:
 
 	void UpdateNumPixelsPerSample()
 	{
-		NumPixelsPerSample = FMath::Max( (float)MIN_NUM_PIXELS_PER_SAMPLE, (float)ThisGeometry.GetLocalSize().X / AllFrames.Num() );
+		NumPixelsPerSample = FMath::Max( (float)MIN_NUM_PIXELS_PER_SAMPLE, (float)ThisGeometry.GetLocalSize().X / (float)AllFrames.Num() );
 	}
 
 	const float GetNumPixelsPerSample() const
