Fixed more "implicit conversion" errors in Profiler module.

#rnx

diff --git a/Engine/Source/Developer/Profiler/Private/Widgets/SDataGraph.cpp b/Engine/Source/Developer/Profiler/Private/Widgets/SDataGraph.cpp
--- a/Engine/Source/Developer/Profiler/Private/Widgets/SDataGraph.cpp
+++ b/Engine/Source/Developer/Profiler/Private/Widgets/SDataGraph.cpp
@@ -575,7 +575,7 @@ int32 SDataGraph::OnPaint( const FPaintArgs& Args, const FGeometry& AllottedGeom
 			for( int32 FrameIndex = FrameStartIndex; FrameIndex < FrameEndIndex; FrameIndex += AvgFrameRate )
 			{
 				const float MarkerPosX = static_cast<float>((FrameIndex - GraphOffset) * DistanceBetweenPoints);
-				const float ElapsedFrameTimeMS = FrameIndex * FTimeAccuracy::AsFrameTime( TimeBasedAccuracy );
+				const float ElapsedFrameTimeMS = (float)FrameIndex * FTimeAccuracy::AsFrameTime( TimeBasedAccuracy );
 				const int32 ElapsedFrameTime = FMath::Max( FMath::RoundToInt( ElapsedFrameTimeMS * 0.001f ) - 1, 0 );
 				const int32 AccumulatedFrameCounter = bCanBeDisplayedAsMulti ? FrameIndex : DataProvider->GetAccumulatedFrameCounter(ElapsedFrameTime);
 
@@ -657,7 +657,7 @@ int32 SDataGraph::OnPaint( const FPaintArgs& Args, const FGeometry& AllottedGeom
 
 	const float MinTimeValue = 0.0f;
 	const float MaxTimeValue = ScaleY;
-	const float TimeValueGraphScale = MaxTimeValue / SecondaryIndicators;
+	const float TimeValueGraphScale = MaxTimeValue / (float)SecondaryIndicators;
 	const float TimeValueToGraph = static_cast<float>(AllottedGeometry.GetLocalSize().Y) / MaxTimeValue;
 
 	for( int32 SecondaryIndex = 1; SecondaryIndex <= SecondaryIndicators; SecondaryIndex++ )
@@ -824,8 +824,8 @@ int32 SDataGraph::OnPaint( const FPaintArgs& Args, const FGeometry& AllottedGeom
 	{
 		LayerId++;
 		const float LocalGraphOffset = static_cast<float>(GraphOffset * DistanceBetweenPoints);
-		const float LocalGraphSelectionX0 = FrameIndices[0]*DistanceBetweenPoints - LocalGraphOffset;
-		const float LocalGraphSelectionX1 = FrameIndices[1]*DistanceBetweenPoints - LocalGraphOffset;
+		const float LocalGraphSelectionX0 = static_cast<float>(FrameIndices[0] * DistanceBetweenPoints) - LocalGraphOffset;
+		const float LocalGraphSelectionX1 = static_cast<float>(FrameIndices[1] * DistanceBetweenPoints) - LocalGraphOffset;
 		const float LocalGraphSelectionX[2] = { LocalGraphSelectionX0, LocalGraphSelectionX1 };
 
 		const uint32 NumVisibleFrameMarkers = ( FrameIndices[0]==FrameIndices[1] ) ? 1 : 2;
@@ -1076,7 +1076,7 @@ FReply SDataGraph::OnMouseMove( const FGeometry& MyGeometry, const FPointerEvent
 	FReply Reply = FReply::Unhandled();
 	MousePosition = MyGeometry.AbsoluteToLocal( MouseEvent.GetScreenSpacePosition() );
 	HoveredFrameIndex = CalculateFrameIndex( MousePosition );
-	HoveredFrameStartTimeMS = HoveredFrameIndex * FTimeAccuracy::AsFrameTime( TimeBasedAccuracy );
+	HoveredFrameStartTimeMS = (float)HoveredFrameIndex * FTimeAccuracy::AsFrameTime( TimeBasedAccuracy );
 
 	if( MouseEvent.IsMouseButtonDown(EKeys::LeftMouseButton) )
 	{
@@ -1095,7 +1095,7 @@ FReply SDataGraph::OnMouseMove( const FGeometry& MyGeometry, const FPointerEvent
 		if( HasMouseCapture() && !MouseEvent.GetCursorDelta().IsZero() )
 		{
 			bIsRMB_Scrolling = true;
-			const float ScrollByAmount = -static_cast<float>(MouseEvent.GetCursorDelta().X) * (1.0f/DistanceBetweenPoints);
+			const float ScrollByAmount = -static_cast<float>(MouseEvent.GetCursorDelta().X) * (1.0f/(float)DistanceBetweenPoints);
 			RealGraphOffset += ScrollByAmount;
 
 			GraphOffset = FMath::Clamp( FMath::TruncToInt( RealGraphOffset ), 0, FMath::Max(NumDataPoints-NumVisiblePoints,0) );
@@ -1110,7 +1110,7 @@ FReply SDataGraph::OnMouseMove( const FGeometry& MyGeometry, const FPointerEvent
 
 const int32 SDataGraph::CalculateFrameIndex( const FVector2D InMousePosition ) const
 {
-	const float ScaleX = 1.0f/DistanceBetweenPoints;
+	const float ScaleX = 1.0f/(float)DistanceBetweenPoints;
 	const int32 MousePositionOffset = FMath::TruncToInt( (static_cast<float>(InMousePosition.X) + HalfGraphMarkerWidth) * ScaleX );
 	return FMath::Clamp( GraphOffset+MousePositionOffset, 0, NumDataPoints-1 );
 }
diff --git a/Engine/Source/Developer/Profiler/Private/Widgets/SHistogram.cpp b/Engine/Source/Developer/Profiler/Private/Widgets/SHistogram.cpp
--- a/Engine/Source/Developer/Profiler/Private/Widgets/SHistogram.cpp
+++ b/Engine/Source/Developer/Profiler/Private/Widgets/SHistogram.cpp
@@ -59,7 +59,7 @@ int32 SHistogram::OnPaint( const FPaintArgs& Args, const FGeometry& AllottedGeom
 	const float LabelBuffer = 25.0f;
 
 	// draw the grid lines
-	uint32 CountX = (uint32)((AllottedGeometry.Size.X-LabelBuffer*2.0f) / Description.GetBinCount());
+	uint32 CountX = (uint32)((AllottedGeometry.Size.X-LabelBuffer*2.0f) / (float)Description.GetBinCount());
 	float StartX = LabelBuffer;
 	static const FLinearColor GridColor = FLinearColor(0.0f,0.0f,0.0f, 0.25f);
 	static const FLinearColor GridTextColor = FLinearColor(1.0f,1.0f,1.0f, 0.25f);
@@ -69,11 +69,11 @@ int32 SHistogram::OnPaint( const FPaintArgs& Args, const FGeometry& AllottedGeom
 	TArray<FVector2D> LinePoints;
 
 	// draw the histogram box
-	LinePoints.Add( FVector2D(StartX-1, LabelBuffer-1) );
-	LinePoints.Add( FVector2D(StartX + Description.GetBinCount()*CountX+1, LabelBuffer-1) );
-	LinePoints.Add( FVector2D(StartX + Description.GetBinCount()*CountX+1, AllottedGeometry.GetLocalSize().Y - LabelBuffer+1) );
-	LinePoints.Add( FVector2D(StartX-1, AllottedGeometry.Size.Y - LabelBuffer+1) );
-	LinePoints.Add( FVector2D(StartX-1, LabelBuffer-1) );
+	LinePoints.Add( FVector2D(StartX - 1, LabelBuffer - 1) );
+	LinePoints.Add( FVector2D(StartX + (float)Description.GetBinCount() * (float)CountX + 1, LabelBuffer - 1) );
+	LinePoints.Add( FVector2D(StartX + (float)Description.GetBinCount() * (float)CountX + 1, AllottedGeometry.GetLocalSize().Y - LabelBuffer + 1) );
+	LinePoints.Add( FVector2D(StartX - 1, AllottedGeometry.Size.Y - LabelBuffer + 1) );
+	LinePoints.Add( FVector2D(StartX - 1, LabelBuffer - 1) );
 	FSlateDrawElement::MakeLines(
 		OutDrawElements,
 		LayerId,
@@ -88,7 +88,7 @@ int32 SHistogram::OnPaint( const FPaintArgs& Args, const FGeometry& AllottedGeom
 	// draw the vertical lines
 	for (int32 Index = 0; Index < Description.GetBinCount(); ++Index)
 	{
-		float MarkerPosX = StartX + Index * CountX;
+		float MarkerPosX = StartX + (float)(Index * CountX);
 		LinePoints.Add( FVector2D(MarkerPosX, LabelBuffer-1) );
 		LinePoints.Add( FVector2D(MarkerPosX, AllottedGeometry.GetLocalSize().Y - LabelBuffer+1) );
 		FSlateDrawElement::MakeLines(
@@ -102,7 +102,7 @@ int32 SHistogram::OnPaint( const FPaintArgs& Args, const FGeometry& AllottedGeom
 		LinePoints.Empty();
 
 		// Bottom - X-Axes numbers, starting from MinValue
-		const FString XLabel = FString::Printf(TEXT("%.0f"), Description.MinValue + Index*Description.Interval);
+		const FString XLabel = FString::Printf(TEXT("%.0f"), Description.MinValue + (float)Index * Description.Interval);
 		float FontCharWidth = static_cast<float>(FontMeasureService->Measure(XLabel, SummaryFont).X);
 		FSlateDrawElement::MakeText(
 			OutDrawElements, 
@@ -122,9 +122,9 @@ int32 SHistogram::OnPaint( const FPaintArgs& Args, const FGeometry& AllottedGeom
 	float StartY = LabelBuffer;
 	for (int32 Index = 0; Index < 5; ++Index)
 	{
-		float MarkerPosY = StartY + Index * CountY;
+		float MarkerPosY = StartY + (float)Index * CountY;
 		LinePoints.Add( FVector2D(StartX, MarkerPosY) );
-		LinePoints.Add( FVector2D(StartX + Description.GetBinCount()*CountX, MarkerPosY) );
+		LinePoints.Add( FVector2D(StartX + (float)(Description.GetBinCount() * CountX), MarkerPosY) );
 		FSlateDrawElement::MakeLines(
 			OutDrawElements,
 			LayerId,
@@ -154,7 +154,7 @@ int32 SHistogram::OnPaint( const FPaintArgs& Args, const FGeometry& AllottedGeom
 
 	for (int32 Index = 0; Index < Description.GetBinCount(); ++Index)
 	{
-		float MarkerPosX = StartX + Index * CountX;
+		float MarkerPosX = StartX + (float)(Index * CountX);
 		float SizeY = (float)Description.GetCount(Index) / (float)Description.GetTotalCount() * ((float)AllottedGeometry.GetLocalSize().Y - LabelBuffer*2.0f);
 		FSlateDrawElement::MakeBox(
 			OutDrawElements,
diff --git a/Engine/Source/Developer/Profiler/Private/Widgets/SHistogram.h b/Engine/Source/Developer/Profiler/Private/Widgets/SHistogram.h
--- a/Engine/Source/Developer/Profiler/Private/Widgets/SHistogram.h
+++ b/Engine/Source/Developer/Profiler/Private/Widgets/SHistogram.h
@@ -56,8 +56,8 @@ public:
 	/** Retrieves the count for the specified bin */
 	int32 GetCount(int32 Bin) const
 	{
-		float MinVal = MinValue + Bin * Interval;
-		float MaxVal = MinValue + (Bin+1) * Interval;
+		float MinVal = MinValue + (float)Bin * Interval;
+		float MaxVal = MinValue + (float)(Bin+1) * Interval;
 		return HistogramDataSource->GetCount(MinVal, MaxVal);
 	}
 
diff --git a/Engine/Source/Developer/Profiler/Private/Widgets/SProfilerGraphPanel.cpp b/Engine/Source/Developer/Profiler/Private/Widgets/SProfilerGraphPanel.cpp
--- a/Engine/Source/Developer/Profiler/Private/Widgets/SProfilerGraphPanel.cpp
+++ b/Engine/Source/Developer/Profiler/Private/Widgets/SProfilerGraphPanel.cpp
@@ -132,7 +132,7 @@ void SProfilerGraphPanel::HorizontalScrollBar_OnUserScrolled(float ScrollOffset)
 
 		HorizontalScrollBar->SetState(ScrollOffset, ThumbSizeFraction);
 
-		GraphOffset = FMath::TruncToInt(ScrollOffset * NumDataPoints);
+		GraphOffset = FMath::TruncToInt(ScrollOffset * (float)NumDataPoints);
 		DataGraph->ScrollTo(GraphOffset);
 
 		ProfilerMiniView.Pin()->MoveWithoutZoomSelectionBox(GraphOffset, GraphOffset + NumVisiblePoints);
diff --git a/Engine/Source/Developer/Profiler/Private/Widgets/SProfilerMiniView.cpp b/Engine/Source/Developer/Profiler/Private/Widgets/SProfilerMiniView.cpp
--- a/Engine/Source/Developer/Profiler/Private/Widgets/SProfilerMiniView.cpp
+++ b/Engine/Source/Developer/Profiler/Private/Widgets/SProfilerMiniView.cpp
@@ -674,7 +674,7 @@ void SProfilerMiniView::ProcessData()
 		for( int32 FrameIndex = 0; FrameIndex < NumAllFrames; ++FrameIndex )
 		{
 			const FFrameThreadTimes& FrameThreadTimes = AllFrames[FrameIndex];
-			const int32 SampleIndex = FMath::TruncToInt( ScaleRatio*FrameIndex );
+			const int32 SampleIndex = FMath::TruncToInt( ScaleRatio * (float)FrameIndex );
 
 			FMiniViewSample& Dest = MiniViewSamples[SampleIndex];
 			Dest.AddFrameAndFindMax( FrameThreadTimes );
