Another fix for "implicit conversion" errors in Profiler module.

#rnx

diff --git a/Engine/Source/Developer/Profiler/Private/ProfilerDataSource.h b/Engine/Source/Developer/Profiler/Private/ProfilerDataSource.h
--- a/Engine/Source/Developer/Profiler/Private/ProfilerDataSource.h
+++ b/Engine/Source/Developer/Profiler/Private/ProfilerDataSource.h
@@ -322,7 +322,7 @@ public:
 
 			for( uint32 NewValueIndex = ChunkStartIndex; NewValueIndex < ChunkEndIndex; NewValueIndex++ )
 			{
-				const float SampleStartTimeMS = NewValueIndex * TimeAccuracyMS;
+				const float SampleStartTimeMS = (float)NewValueIndex * TimeAccuracyMS;
 				ThisCacheDataContainer::CachedValues(NewValueIndex) = (Type)static_cast<const ManagerClass*>(this)->GetUncachedValueFromTimeRange( SampleStartTimeMS, SampleStartTimeMS+TimeAccuracyMS );
 			}
 
